fix for VulkanShaderFormat being compiled on Mac

diff --git a/Engine/Source/Developer/VulkanShaderFormat/VulkanShaderFormat.Build.cs b/Engine/Source/Developer/VulkanShaderFormat/VulkanShaderFormat.Build.cs
--- a/Engine/Source/Developer/VulkanShaderFormat/VulkanShaderFormat.Build.cs
+++ b/Engine/Source/Developer/VulkanShaderFormat/VulkanShaderFormat.Build.cs
@@ -20,5 +20,10 @@ public class VulkanShaderFormat : ModuleRules
 
         AddEngineThirdPartyPrivateStaticDependencies(Target, "HLSLCC");
 		AddEngineThirdPartyPrivateStaticDependencies(Target, "GlsLang");
+
+		if (Target.Platform != UnrealTargetPlatform.Win64 || Target.Platform != UnrealTargetPlatform.Win32 || Target.Platform != UnrealTargetPlatform.Android)
+		{
+			PrecompileForTargets = PrecompileTargetsType.None;
+		}
 	}
 }
