Fix for shadow variable warning from GLSL (since it's a third party library, I'm just disabling the warning in our client)

diff --git a/Engine/Source/Developer/VulkanShaderFormat/VulkanShaderFormat.Build.cs b/Engine/Source/Developer/VulkanShaderFormat/VulkanShaderFormat.Build.cs
--- a/Engine/Source/Developer/VulkanShaderFormat/VulkanShaderFormat.Build.cs
+++ b/Engine/Source/Developer/VulkanShaderFormat/VulkanShaderFormat.Build.cs
@@ -20,6 +20,9 @@ public class VulkanShaderFormat : ModuleRules
 			}
 			);
 
+        // GlsLang has a shadowed variable:
+        bEnableShadowVariableWarnings = false;
+
         AddEngineThirdPartyPrivateStaticDependencies(Target, "HLSLCC");
 		AddEngineThirdPartyPrivateStaticDependencies(Target, "GlsLang");
 
