Copying //UE4/Dev-Rendering to //UE4/Dev-Main (Source: //UE4/Dev-Rendering @ 3054480)

==========================
MAJOR FEATURES + CHANGES
==========================

Change 3045482 on 2016/07/11 by [REDACTED_PERSON]

	DX12 Quries need to individually track their syncpoints. Only when resolving a query on the same frame should be stall.

Change 3045929 on 2016/07/12 by [REDACTED_PERSON]

	Removing some deprecated node types from Niagara

Change 3045951 on 2016/07/12 by [REDACTED_PERSON]

	D3D11 Log detailed live device info on shutdown if the debug layer is enabled (including resource types)

Change 3046019 on 2016/07/12 by [REDACTED_PERSON]

	Fixed typo in material input name.

Change 3046053 on 2016/07/12 by [REDACTED_PERSON]

	DR - Fix GL4 shutdown

Change 3046055 on 2016/07/12 by [REDACTED_PERSON]

	DR - vk - Fix NumInstances=0

Change 3046063 on 2016/07/12 by [REDACTED_PERSON]

	DR - vk - Added flat to uint layouts per glslang
	- Fix bad extension on dumped shaders

Change 3046067 on 2016/07/12 by [REDACTED_PERSON]

	DR - vk - Fix check when not using color RT
	- Added queue submit & present counters

Change 3046088 on 2016/07/12 by [REDACTED_PERSON]

	Live GPU stats
	A non-hierarchical realtime high level GPU profiler with support for cumulative stat recording.
	Stats are added with SCOPED_GPU_STAT macros, e.g. SCOPED_GPU_STAT(RHICmdList, Stat_GPU_Distortion)
	The bulk of the files in this change are simply instrumentation for the renderer. The core changes are in SceneUtils.cpp/h and D3D11Query.cpp (this is the XB1/DX11X implementation of timestamp RHI queries, which was missing)
	Note: this is currently disabled by default. Enable with the cvar r.gpustatsenabled
	Tested on PC, XB1, PS4

Change 3046128 on 2016/07/12 by [REDACTED_PERSON]

	Max draw distance and fade range for lights, requested by JonL

Change 3046183 on 2016/07/12 by [REDACTED_PERSON]

	PR #2532: Fix SSAO being applied in unlit viewmode (Contributed by [REDACTED_PERSON])

Change 3046223 on 2016/07/12 by [REDACTED_PERSON]

	Fix Scene Cube Captures. SceneCaptureSource flag on the ViewFamily was not set for cube components.

Change 3046228 on 2016/07/12 by [REDACTED_PERSON]

	Add Voronoi noise to Noise material node.

	Four versions with differing speed/quality levels accessed through the Quality value in the material node. Tooltips give estimates of the cost of each.

	Also includes spiffy new Rand3DPCG16 and Rand3DPCG32 int3 to int3 hash functions, and a 20% improvement on the computed gradient noise.

Change 3046269 on 2016/07/12 by [REDACTED_PERSON]

	DR - Skip flush on RHIDiscardRenderTargets and only use it on platforms that need it (ie OpenGL)

Change 3046294 on 2016/07/12 by [REDACTED_PERSON]

	DR - Fix static analyisis
	warning C6326: Potential comparison of a constant with another constant.

Change 3046295 on 2016/07/12 by [REDACTED_PERSON]

	DR - Fix the previous fix

Change 3046731 on 2016/07/12 by [REDACTED_PERSON]

	Fix typo in shader random number constant: repeated extra digit made it too big.

Change 3046796 on 2016/07/12 by [REDACTED_PERSON]

	The texture streaming manager now keeps a set of all valid textures.
	This is used to prevent from indirecting deleted memory upon SetTexturesRemovedTimestamp.

Change 3046800 on 2016/07/12 by [REDACTED_PERSON]

	DR - vk - Added create image & renderpass dump

Change 3046845 on 2016/07/12 by [REDACTED_PERSON]

	Forgot to apply MaxGPUSkinBones Cvar access changes in a few locations.

Change 3047023 on 2016/07/12 by [REDACTED_PERSON]

	Niagara:
	-a bit of cleanup
	-now store and double buffer attributes individually, eliminating unnecessary copy of unused attributes
	-removed FNiagaraConstantMap, replaced with an instance of FNiagaraConstants
	-some code simplification
	-removed some deprecated structs and code used only by old content

Change 3047052 on 2016/07/12 by [REDACTED_PERSON]

	Unshelved from pending changelist '3044062':

	PR #2588: Adding blend mode BLEND_AlphaComposite (4.12) (Contributed by [REDACTED_PERSON])

Change 3047727 on 2016/07/13 by [REDACTED_PERSON]

	Fix Scene Capture Components only updating every other frame.

Change 3047919 on 2016/07/13 by [REDACTED_PERSON]

	CMask decode, use in deferred decals, for PS4

Change 3047921 on 2016/07/13 by [REDACTED_PERSON]

	"Build Texture Streaming" will now remove duplicate error msg when computing texcoord scales.
	Also, several texture messages are packed on the same line if they relate to the same material.

Change 3047952 on 2016/07/13 by [REDACTED_PERSON]

	DR - vk - Initial prep pass for separating combined images & samplers

Change 3048648 on 2016/07/13 by [REDACTED_PERSON]

	Fix rare GPU hang when asynctexture reallocs would overlap with EndFrame

Change 3049058 on 2016/07/13 by [REDACTED_PERSON]

	DR - vk - timestamps

Change 3049725 on 2016/07/14 by [REDACTED_PERSON]

	Fix autosdk bug where not having a platform directory sync'd at all would break manual SDK detection

Change 3049742 on 2016/07/14 by [REDACTED_PERSON]

	DR - Fix warning

Change 3049902 on 2016/07/14 by [REDACTED_PERSON]

	DR - Fix typo

Change 3050345 on 2016/07/14 by [REDACTED_PERSON]

	UE-23925
	Clamping noise tessellation for beams at a high but sensible value; also making sure during beam index buffer building that we never get over 2^16 indices; this is a bit hokey, but there are so many variables that can influence triangle/index count, that this is the only way to be sure (short of nuking the entire site from orbit).

Change 3050409 on 2016/07/14 by [REDACTED_PERSON]

	Replicating 3049049; missing break and check for active particles when resolving a source point to avoid a potential crash

Change 3050809 on 2016/07/14 by [REDACTED_PERSON]

	DR - vk - Remove redundant validation layers

Change 3051319 on 2016/07/15 by [REDACTED_PERSON]

	Fix for world space camera position not being exposed in decal pixel shaders; also fixes decal lighting missing spec and reflection
	The fix was to calculate ResolvedView at the top of the shader. Previously this was not initialized

Change 3051692 on 2016/07/15 by [REDACTED_PERSON]

	DR - vk - Enable RHI thread by default

Change 3052103 on 2016/07/15 by [REDACTED_PERSON]

	Disabled depth offset in depth only pixel shaders when using debug view shaders (to prevent Z fighting).

Change 3052140 on 2016/07/15 by [REDACTED_PERSON]

	DR - vk - Fix shader snafu

Change 3052495 on 2016/07/15 by [REDACTED_PERSON]

	DR - Fix for Win32 compile

Change 3052536 on 2016/07/15 by [REDACTED_PERSON]

	Fixed texture streaming overbudget warning when using per texture bias.

diff --git a/Engine/Config/BaseInput.ini b/Engine/Config/BaseInput.ini
--- a/Engine/Config/BaseInput.ini
+++ b/Engine/Config/BaseInput.ini
@@ -97,6 +97,7 @@ MaxScrollbackSize=1024
 +ManualAutoCompleteList=(Command="Stat PHYSICS",Desc="Displays physics performance stats")
 +ManualAutoCompleteList=(Command="Stat STREAMING",Desc="Displays basic texture streaming stats")
 +ManualAutoCompleteList=(Command="Stat STREAMINGDETAILS",Desc="Displays detailed texture streaming stats")
++ManualAutoCompleteList=(Command="Stat GPU",Desc="Displays GPU stats for the frame")
 +ManualAutoCompleteList=(Command="Stat COLLISION",Desc=)
 +ManualAutoCompleteList=(Command="Stat PARTICLES",Desc=)
 +ManualAutoCompleteList=(Command="Stat SCRIPT",Desc=)
diff --git a/Engine/Plugins/2D/Paper2D/Source/Paper2DEditor/Private/SpriteEditor/SpriteEditorViewportClient.cpp b/Engine/Plugins/2D/Paper2D/Source/Paper2DEditor/Private/SpriteEditor/SpriteEditorViewportClient.cpp
--- a/Engine/Plugins/2D/Paper2D/Source/Paper2DEditor/Private/SpriteEditor/SpriteEditorViewportClient.cpp
+++ b/Engine/Plugins/2D/Paper2D/Source/Paper2DEditor/Private/SpriteEditor/SpriteEditorViewportClient.cpp
@@ -362,6 +362,7 @@ void FSpriteEditorViewportClient::AnalyzeSpriteMaterialType(UPaperSprite* Sprite
 				case EBlendMode::BLEND_Translucent:
 				case EBlendMode::BLEND_Additive:
 				case EBlendMode::BLEND_Modulate:
+				case EBlendMode::BLEND_AlphaComposite:
 					NumTranslucentTriangles += NumTriangles;
 					break;
 				case EBlendMode::BLEND_Masked:
diff --git a/Engine/Plugins/2D/Paper2D/Source/Paper2DEditor/Private/TileMapEditing/TileMapEditorViewportClient.cpp b/Engine/Plugins/2D/Paper2D/Source/Paper2DEditor/Private/TileMapEditing/TileMapEditorViewportClient.cpp
--- a/Engine/Plugins/2D/Paper2D/Source/Paper2DEditor/Private/TileMapEditing/TileMapEditorViewportClient.cpp
+++ b/Engine/Plugins/2D/Paper2D/Source/Paper2DEditor/Private/TileMapEditing/TileMapEditorViewportClient.cpp
@@ -189,6 +189,7 @@ void FTileMapEditorViewportClient::DrawCanvas(FViewport& InViewport, FSceneView&
 					case EBlendMode::BLEND_Translucent:
 					case EBlendMode::BLEND_Additive:
 					case EBlendMode::BLEND_Modulate:
+					case EBlendMode::BLEND_AlphaComposite:
 						MaterialType = Translucent;
 						break;
 					case EBlendMode::BLEND_Masked:
diff --git a/Engine/Shaders/BasePassPixelShader.usf b/Engine/Shaders/BasePassPixelShader.usf
--- a/Engine/Shaders/BasePassPixelShader.usf
+++ b/Engine/Shaders/BasePassPixelShader.usf
@@ -46,6 +46,10 @@ SamplerState	HZBSampler;
 Texture2D		PrevSceneColor;
 SamplerState	PrevSceneColorSampler;
 
+#if PLATFORM_SUPPORTS_RENDERTARGET_WRITE_MASK && USE_DBUFFER && MATERIALDECALRESPONSEMASK && !MATERIALBLENDING_ANY_TRANSLUCENT
+	Texture2D<uint> DBufferMask;
+#endif
+
 #ifndef COMPILER_GLSL
 	#define COMPILER_GLSL 0
 #endif
@@ -789,18 +793,26 @@ void FPixelShaderInOut_MainPS(
 #endif
 	if(Primitive.DecalReceiverMask > 0 && View.ShowDecalsMask > 0)
 	{
-		float2 NDC = MaterialParameters.ScreenPosition.xy / MaterialParameters.ScreenPosition.w;
-		
-		float2 ScreenUV = NDC * View.ScreenPositionScaleBias.xy + View.ScreenPositionScaleBias.wz;
-		
-		FDBufferData DBufferData = GetDBufferData(ScreenUV);
 
-		// the material can disable the DBuffer effects for better performance or control
-		if((MATERIALDECALRESPONSEMASK & 0x1) == 0) { DBufferData.PreMulColor = 0; DBufferData.ColorOpacity = 1; }
-		if((MATERIALDECALRESPONSEMASK & 0x2) == 0) { DBufferData.PreMulWorldNormal = 0; DBufferData.NormalOpacity = 1; }
-		if((MATERIALDECALRESPONSEMASK & 0x4) == 0) { DBufferData.PreMulRoughness = 0; DBufferData.RoughnessOpacity = 1; }
+	#if PLATFORM_SUPPORTS_RENDERTARGET_WRITE_MASK
+		uint RTWriteMaskBit = DecodeRTWriteMaskTexture(In.SvPosition.xy, DBufferMask);
+		if(RTWriteMaskBit)
+	#endif
+		{
+			float2 NDC = MaterialParameters.ScreenPosition.xy / MaterialParameters.ScreenPosition.w;
+
+			// Note: We are using View and not ResolvedView here.
+			// It has the correct ScreenPositionScaleBias values for screen space compositing.
+			float2 ScreenUV = NDC * View.ScreenPositionScaleBias.xy + View.ScreenPositionScaleBias.wz;
 
-		ApplyDBufferData(DBufferData, MaterialParameters.WorldNormal, SubsurfaceColor, Roughness, BaseColor, Metallic, Specular);
+			FDBufferData DBufferData = GetDBufferData(ScreenUV);
+
+			// the material can disable the DBuffer effects for better performance or control
+			if((MATERIALDECALRESPONSEMASK & 0x1) == 0) { DBufferData.PreMulColor = 0; DBufferData.ColorOpacity = 1; }
+			if((MATERIALDECALRESPONSEMASK & 0x2) == 0) { DBufferData.PreMulWorldNormal = 0; DBufferData.NormalOpacity = 1; }
+			if((MATERIALDECALRESPONSEMASK & 0x4) == 0) { DBufferData.PreMulRoughness = 0; DBufferData.RoughnessOpacity = 1; }
+			ApplyDBufferData(DBufferData, MaterialParameters.WorldNormal, SubsurfaceColor, Roughness, BaseColor, Metallic, Specular);
+		}
 	}
 #endif
 
diff --git a/Engine/Shaders/Common.usf b/Engine/Shaders/Common.usf
--- a/Engine/Shaders/Common.usf
+++ b/Engine/Shaders/Common.usf
@@ -954,6 +954,24 @@ MaterialFloat3 Decode32BPPHDR(MaterialFloat4 Encoded, MaterialFloat3 OtherEncode
 	}
 }
 
+
+/** Get render target write mask value
+  * This gets a bit from a write mask texture created with FRTWriteMaskDecodeCS. Only supprted on some platforms.
+  */
+#if PLATFORM_SUPPORTS_RENDERTARGET_WRITE_MASK
+uint DecodeRTWriteMaskTexture(in float2 ScreenPosition, in Texture2D<uint> RTWriteMaskTexture)
+{
+	int2 IntPosition = int2(ScreenPosition.xy);
+	uint RTWriteMaskValue = RTWriteMaskTexture.Load( int3(IntPosition.x/8, IntPosition.y/8, 0) );
+
+	int2 BitCoord = ((IntPosition / int2(4, 4)) % int2(2, 2));
+	uint BitIdx = BitCoord.x + (BitCoord.y*2);
+	uint RTWriteMaskBit = RTWriteMaskValue & (1<<BitIdx);
+
+	return RTWriteMaskBit;
+}
+#endif
+
 /** Calculates the ScreenUV given the screen position and an offset fraction. */
 float2 CalcScreenUVFromOffsetFraction(float4 ScreenPosition, float2 OffsetFraction)
 {
@@ -1278,7 +1296,7 @@ float AntialiasedTextureMask( Texture2D Tex, SamplerState Sampler, float2 UV, fl
 	return Result;
 }
 
-float Noise3D_Multiplexer(int Function, float3 Position, bool bTiling, uint RepeatSize)
+float Noise3D_Multiplexer(int Function, float3 Position, int Quality, bool bTiling, uint RepeatSize)
 {
 	// verified, HLSL compiled out the switch if Function is a constant
 	switch(Function)
@@ -1291,8 +1309,10 @@ float Noise3D_Multiplexer(int Function, float3 Position, bool bTiling, uint Repe
 			return FastGradientPerlinNoise3D_TEX(Position);
 		case 3:
 			return GradientNoise3D_ALU(Position, bTiling, RepeatSize);
-		default:
+		case 4:
 			return ValueNoise3D_ALU(Position, bTiling, RepeatSize);
+		default:
+			return VoronoiNoise3D_ALU(Position, Quality, bTiling, RepeatSize);
 	}
 	return 0;
 }
@@ -1315,11 +1335,11 @@ float MaterialExpressionNoise(float3 Position, float Scale, int Quality, int Fun
 
 		if(bTurbulence)
 		{
-			Out += abs(Noise3D_Multiplexer(Function, Position, bTiling, RepeatSize)) * OutScale;
+			Out += abs(Noise3D_Multiplexer(Function, Position, Quality, bTiling, RepeatSize)) * OutScale;
 		}
 		else
 		{
-			Out += Noise3D_Multiplexer(Function, Position, bTiling, RepeatSize) * OutScale;
+			Out += Noise3D_Multiplexer(Function, Position, Quality, bTiling, RepeatSize) * OutScale;
 		}
 
 		Position *= LevelScale;
diff --git a/Engine/Shaders/DeferredDecal.usf b/Engine/Shaders/DeferredDecal.usf
--- a/Engine/Shaders/DeferredDecal.usf
+++ b/Engine/Shaders/DeferredDecal.usf
@@ -149,6 +149,8 @@ float4 SvPositionToScreenPosition2(float4 SvPosition)
 // is called in MainPS() from PixelShaderOutputCommon.usf
 void FPixelShaderInOut_MainPS(inout FPixelShaderIn In, inout FPixelShaderOut Out)
 {
+	ResolvedView = ResolveView();
+
 	float2 ScreenUV = SvPositionToBufferUV(In.SvPosition);
 
 	// make SvPosition appear to be rasterized with the depth from the depth buffer
diff --git a/Engine/Shaders/Definitions.usf b/Engine/Shaders/Definitions.usf
--- a/Engine/Shaders/Definitions.usf
+++ b/Engine/Shaders/Definitions.usf
@@ -34,6 +34,10 @@
 #define MATERIALBLENDING_MODULATE						0
 #endif
 
+#ifndef MATERIALBLENDING_ALPHACOMPOSITE
+#define MATERIALBLENDING_ALPHACOMPOSITE					0
+#endif
+
 #ifndef MATERIAL_SHADINGMODEL_DEFAULT_LIT
 #define MATERIAL_SHADINGMODEL_DEFAULT_LIT				0
 #endif
diff --git a/Engine/Shaders/DepthOnlyPixelShader.usf b/Engine/Shaders/DepthOnlyPixelShader.usf
--- a/Engine/Shaders/DepthOnlyPixelShader.usf
+++ b/Engine/Shaders/DepthOnlyPixelShader.usf
@@ -8,6 +8,10 @@
 #include "Material.usf"
 #include "VertexFactory.usf"
 
+#if OUTPUT_PIXEL_DEPTH_OFFSET
+bool bApplyDepthOffset;
+#endif
+
 void Main(
 #if !MATERIALBLENDING_SOLID || OUTPUT_PIXEL_DEPTH_OFFSET
 	FVertexFactoryInterpolantsVSToPS FactoryInterpolants,
@@ -32,6 +36,7 @@ void Main(
 
 	#if OUTPUT_PIXEL_DEPTH_OFFSET
 	ApplyPixelDepthOffsetToMaterialParameters(MaterialParameters, PixelMaterialInputs, OutDepth);
+	OutDepth = bApplyDepthOffset ? OutDepth : SvPosition.z;
 	#endif
 
 	GetMaterialCoverageAndClipping(MaterialParameters, PixelMaterialInputs);
diff --git a/Engine/Shaders/MobileBasePassPixelShader.usf b/Engine/Shaders/MobileBasePassPixelShader.usf
--- a/Engine/Shaders/MobileBasePassPixelShader.usf
+++ b/Engine/Shaders/MobileBasePassPixelShader.usf
@@ -272,6 +272,10 @@ half3 FrameBufferBlendOp(half4 Source)
 	return Source.rgb;
 #elif MATERIALBLENDING_MASKED
 	return Source.rgb;
+// AlphaComposite will set both MATERIALBLENDING_TRANSLUCENT and MATERIALBLENDING_ALPHACOMPOSITE defines
+// so ensure  MATERIALBLENDING_ALPHACOMPOSITE gets first in line
+#elif MATERIALBLENDING_ALPHACOMPOSITE
+	return Source.rgb + (Dest.rgb*(1.0 - Source.a));
 #elif MATERIALBLENDING_TRANSLUCENT
 	return (Source.rgb*Source.a) + (Dest.rgb*(1.0 - Source.a));
 #elif MATERIALBLENDING_ADDITIVE
diff --git a/Engine/Shaders/Random.usf b/Engine/Shaders/Random.usf
--- a/Engine/Shaders/Random.usf
+++ b/Engine/Shaders/Random.usf
@@ -9,6 +9,7 @@
 
 // @param xy should be a integer position (e.g. pixel position on the screen), repeats each 128x128 pixels
 // similar to a texture lookup but is only ALU
+// ~13 ALU operations (3 frac, 6 *, 4 mad)
 float PseudoRandom(float2 xy)
 {
 	float2 pos = frac(xy / 128.0f) * 128.0f + float2(-64.340622f, -72.465622f);
@@ -21,6 +22,7 @@ float PseudoRandom(float2 xy)
 //note: from "NEXT GENERATION POST PROCESSING IN CALL OF DUTY: ADVANCED WARFARE"
 //      http://advances.realtimerendering.com/s2014/index.html
 // Epic extended by FrameId
+// ~7 ALU operations (2 frac, 3 mad, 2 *)
 // @return 0..1
 float InterleavedGradientNoise( float2 uv, float FrameId )
 {
@@ -32,12 +34,13 @@ float InterleavedGradientNoise( float2 uv, float FrameId )
 }
 
 // [0, 1[
+// ~10 ALU operations (2 frac, 5 *, 3 mad)
 float RandFast( uint2 PixelPos, float Magic = 3571.0 )
 {
-	float2 Random = ( 1.0 / 4320.0 ) * PixelPos + float2( 0.25, 0.0 );
-	Random = frac( dot( Random * Random, Magic ) );
-	Random = frac( dot( Random * Random, Magic ) );
-	return Random.x;
+	float2 Random2 = ( 1.0 / 4320.0 ) * PixelPos + float2( 0.25, 0.0 );
+	float Random = frac( dot( Random2 * Random2, Magic ) );
+	Random = frac( Random * Random * (2 * Magic) );
+	return Random;
 }
 
 // Blum-Blum-Shub-inspired pseudo random number generator
@@ -50,17 +53,13 @@ float RandFast( uint2 PixelPos, float Magic = 3571.0 )
 #define BBS_PRIME32 65521
 
 // Blum-Blum-Shub-inspired pseudo random number generator - float version
-// @param floating point seed - only frac matters (so 0 and 1 produce the same result)
-// @return random number (new seed) in range [0,1)
-// suggested usage: integer seed to fraction output: RandBBSfloat(float(intseed) / BBS_PRIME24)
-// uint version RandBBSint24(intseed) is equivalent to BBS_PRIME24 * RandBBSfloat(float(intseed) / BBS_PRIME24)
-// fractional input works, but for frame-to-frame stability, should not have more than 12 mantissa bits
-//   e.g. RandBBSfloat(floor(seed*4096)/4096)
-//   In theory, should be able to force this rounding by adding then subtracting 4096
-//   Unfortunately, the shading compiler seems too good at eliminating this from the code
+// @param Integer valued floating point seed
+// @return random number in range [0,1)
+// uint version RandBBSint24(intseed) is equivalent to BBS_PRIME24 * RandBBSfloat(float(intseed))
+// ~7 ALU operations (4 *, 3 frac)
 float RandBBSfloat(float seed)
 {
-	float s = frac(seed);
+	float s = frac(seed / BBS_PRIME24);
 	s = frac(s * s * BBS_PRIME24);
 	s = frac(s * s * BBS_PRIME24);
 	return s;
@@ -70,11 +69,12 @@ float RandBBSfloat(float seed)
 // This version exists to match the float one, and because some hardware is faster for 24-bit int ops
 // @param seed - old seed, repeats every BBS_PRIME24 (so 0 and 4093 produce the same result)
 // @return random value (can be used as new seed) in range [0,BBS_PRIME24) = [0,4093)
+// ~5 ALU operations (3 %, 2 *)
 uint RandBBSuint24(uint seed)
 {
 #if (ES2_PROFILE)
 	// no integer mod in ES2
-	return BBS_PRIME24 * RandBBSfloat(float(seed) / BBS_PRIME24);
+	return BBS_PRIME24 * RandBBSfloat(float(seed));
 #else
 	uint s = seed % BBS_PRIME24;
 	s = (s * s) % BBS_PRIME24;
@@ -87,11 +87,12 @@ uint RandBBSuint24(uint seed)
 // uses 65521 as the prime modulus, since it is the largest prime < 2^16
 // @param seed - old seed, repeats every BBS_PRIME32 (so 0 and 65521 produce the same result)
 // @return random value (can be used as new seed) in range [0,BBS_PRIME32) = [0,65521)
+// ~5 ALU operations (3 %, 2 *)
 uint RandBBSuint32(uint seed)
 {
 #if (ES2_PROFILE)
 	// no integer mod in ES2
-	return BBS_PRIME32 * RandBBSfloat(float(seed) / BBS_PRIME24);
+	return BBS_PRIME32 * RandBBSfloat(float(seed));
 #else
 	uint s = seed % BBS_PRIME32;
 	s = (s * s) % BBS_PRIME32;
@@ -100,6 +101,76 @@ uint RandBBSuint32(uint seed)
 #endif
 }
 
+// 3D random number generator inspired by PCGs (permuted congruential generator)
+// Using a **simple** Feistel cipher in place of the usual xor shift permutation step
+// @param v = 3D integer coordinate
+// @return three elements w/ 16 random bits each (0-0xffff).
+// ~8 ALU operations for result.x    (7 mad, 1 >>)
+// ~10 ALU operations for result.xy  (8 mad, 2 >>)
+// ~12 ALU operations for result.xyz (9 mad, 3 >>)
+uint3 Rand3DPCG16(int3 p)
+{
+	// taking a signed int then reinterpreting as unsigned gives good behavior for negatives
+	uint3 v = uint3(p);
+
+	// Linear congruential step. These LCG constants are from Numerical Recipies
+	// For additional #'s, PCG would do multiple LCG steps and scramble each on output
+	// So v here is the RNG state
+	v = v * 1664525u + 1013904223u;
+
+	// PCG uses xorshift for the final shuffle, but it is expensive (and cheap
+	// versions of xorshift have visible artifacts). Instead, use simple MAD Feistel steps
+	//
+	// Feistel ciphers divide the state into separate parts (usually by bits)
+	// then apply a series of permutation steps one part at a time. The permutations
+	// use a reversible operation (usually ^) to part being updated with the result of
+	// a permutation function on the other parts and the key.
+	//
+	// In this case, I'm using v.x, v.y and v.z as the parts, using + instead of ^ for
+	// the combination function, and just multiplying the other two parts (no key) for 
+	// the permutation function.
+	//
+	// That gives a simple mad per round.
+	v.x += v.y*v.z;
+	v.y += v.z*v.x;
+	v.z += v.x*v.y;
+	v.x += v.y*v.z;
+	v.y += v.z*v.x;
+	v.z += v.x*v.y;
+
+	// only top 16 bits are well shuffled
+	return v >> 16u;
+}
+
+// 3D random number generator inspired by PCGs (permuted congruential generator)
+// Using a **simple** Feistel cipher in place of the usual xor shift permutation step
+// @param v = 3D integer coordinate
+// @return three elements w/ 16 random bits each (0-0xffff).
+// ~12 ALU operations for result.x   (10 mad, 3 >>)
+// ~14 ALU operations for result.xy  (11 mad, 3 >>)
+// ~15 ALU operations for result.xyz (12 mad, 3 >>)
+uint3 Rand3DPCG32(int3 p)
+{
+	// taking a signed int then reinterpreting as unsigned gives good behavior for negatives
+	uint3 v = uint3(p);
+
+	// Linear congruential step.
+	v = v * 1664525u + 1013904223u;
+
+	// swapping low and high bits makes all 32 bits pretty good
+	v = v * (1u << 16u) + (v >> 16u);
+
+	// final shuffle
+	v.x += v.y*v.z;
+	v.y += v.z*v.x;
+	v.z += v.x*v.y;
+	v.x += v.y*v.z;
+	v.y += v.z*v.x;
+	v.z += v.x*v.y;
+
+	return v;
+}
+
 /**
  * Find good arbitrary axis vectors to represent U and V axes of a plane,
  * given just the normal. Ported from UnMath.h
@@ -189,7 +260,7 @@ float4 PerlinRamp(float4 t)
 float MGradient(uint seed, float3 offset)
 {
 	uint rand = RandBBSuint24(seed);
-	float3 direction = float3((rand & 1) << 1, (rand & 2), (rand & 4) >> 1) - 1;
+	float3 direction = float3(rand & 1, rand & 2, rand & 4) * float3(2, 1, 0.5) - 1;
 	return dot(direction, offset);
 }
 
@@ -271,14 +342,14 @@ float ValueNoise3D_ALU(float3 v, bool bTiling, float RepeatSize)
 	float seed000, seed001, seed010, seed011, seed100, seed101, seed110, seed111;
 	float3 fv = NoiseSeeds(v, bTiling, RepeatSize, seed000, seed001, seed010, seed011, seed100, seed101, seed110, seed111);
 
-	float rand000 = RandBBSfloat(seed000 / BBS_PRIME24) * 2 - 1;
-	float rand100 = RandBBSfloat(seed100 / BBS_PRIME24) * 2 - 1;
-	float rand010 = RandBBSfloat(seed010 / BBS_PRIME24) * 2 - 1;
-	float rand110 = RandBBSfloat(seed110 / BBS_PRIME24) * 2 - 1;
-	float rand001 = RandBBSfloat(seed001 / BBS_PRIME24) * 2 - 1;
-	float rand101 = RandBBSfloat(seed101 / BBS_PRIME24) * 2 - 1;
-	float rand011 = RandBBSfloat(seed011 / BBS_PRIME24) * 2 - 1;
-	float rand111 = RandBBSfloat(seed111 / BBS_PRIME24) * 2 - 1;
+	float rand000 = RandBBSfloat(seed000) * 2 - 1;
+	float rand100 = RandBBSfloat(seed100) * 2 - 1;
+	float rand010 = RandBBSfloat(seed010) * 2 - 1;
+	float rand110 = RandBBSfloat(seed110) * 2 - 1;
+	float rand001 = RandBBSfloat(seed001) * 2 - 1;
+	float rand101 = RandBBSfloat(seed101) * 2 - 1;
+	float rand011 = RandBBSfloat(seed011) * 2 - 1;
+	float rand111 = RandBBSfloat(seed111) * 2 - 1;
 	
 	float3 Weights = PerlinRamp(float4(fv, 0)).xyz;
 	
@@ -386,6 +457,115 @@ float FastGradientPerlinNoise3D_TEX(float3 xyz)
 	return dot(xyz, n) - d;
 }
 
+
+// 3D jitter offset within a voronoi noise cell
+// @param pos - integer lattice corner
+// @return random offsets vector
+float3 VoronoiCornerSample(int3 pos, int Quality)
+{
+	// random values in [-0.5, 0.5]
+	float3 noise = float3(Rand3DPCG16(pos)) / 0xffff - 0.5;
+
+	// quality level 1 or 2: searches a 2x2x2 neighborhood with points distributed on a sphere
+	// scale factor to guarantee jittered points will be found within a 2x2x2 search
+	if (Quality <= 2)
+	{
+		return normalize(noise) * 0.2588;
+	}
+
+	// quality level 3: searches a 3x3x3 neighborhood with points distributed on a sphere
+	// scale factor to guarantee jittered points will be found within a 3x3x3 search
+	if (Quality == 3)
+	{
+		return normalize(noise) * 0.3090;
+	}
+
+	// quality level 4: jitter to anywhere in the cell, needs 4x4x4 search
+	return noise;
+}
+
+// 220 instruction Worley noise
+float VoronoiNoise3D_ALU(float3 v, int Quality, bool bTiling, float RepeatSize)
+{
+	float3 fv = frac(v),  fv2 = frac(v + 0.5);
+	float3 iv = floor(v), iv2 = floor(v + 0.5);
+
+	// with initial minimum distance = infinity (or at least bigger than 4), first min is optimized away
+	float mindist = 100;
+	float3 offset;
+
+	// quality level 3: do 3x3x3 search centered on current location
+	if (Quality == 3)
+	{
+		float3 mincell = floor(v - 1), maxcell = floor(v + 1);
+		float3 cell;
+		LOOP for (cell.x = mincell.x; cell.x <= maxcell.x; ++cell.x)
+		{
+			LOOP for (cell.y = mincell.y; cell.y <= maxcell.y; ++cell.y)
+			{
+				LOOP for (cell.z = mincell.z; cell.z <= maxcell.z; ++cell.z)
+				{
+					float3 p = v - cell - VoronoiCornerSample(NoiseTileWrap(cell, bTiling, RepeatSize), Quality);
+					mindist = min(mindist, dot(p, p));
+				}
+			}
+		}
+	}
+
+	// all others, do 2x2x2 search (unrolled)
+	else
+	{
+		UNROLL for (offset.x = 0; offset.x <= 1; ++offset.x)
+		{
+			UNROLL for (offset.y = 0; offset.y <= 1; ++offset.y)
+			{
+				UNROLL for (offset.z = 0; offset.z <= 1; ++offset.z)
+				{
+					float3 p = fv - offset - VoronoiCornerSample(NoiseTileWrap(iv + offset, bTiling, RepeatSize), Quality);
+					mindist = min(mindist, dot(p, p));
+					// quality level 2, do extra set of points, offset by half a cell
+					if (Quality == 2)
+					{
+						// 467 is just an offset to a different area in the random number field to avoid similar neighbor artifacts
+						p = fv2 - offset - VoronoiCornerSample(NoiseTileWrap(iv2 + offset, bTiling, RepeatSize) + 467, Quality);
+						mindist = min(mindist, dot(p, p));
+					}
+				}
+			}
+		}
+	}
+
+	// quality level 4: add extra sets of four cells in each direction
+	if (Quality >= 4)
+	{
+		float3 p;
+		UNROLL for (offset.x = -1; offset.x <= 2; offset.x += 3)
+		{
+			UNROLL for (offset.y = 0; offset.y <= 1; ++offset.y)
+			{
+				UNROLL for (offset.z = 0; offset.z <= 1; ++offset.z)
+				{
+					// along x axis
+					p = fv - offset.xyz - VoronoiCornerSample(NoiseTileWrap(iv + offset.xyz, bTiling, RepeatSize), Quality);
+					mindist = min(mindist, dot(p, p));
+
+					// along y axis
+					p = fv - offset.yzx - VoronoiCornerSample(NoiseTileWrap(iv + offset.yzx, bTiling, RepeatSize), Quality);
+					mindist = min(mindist, dot(p, p));
+
+					// along z axis
+					p = fv - offset.zxy - VoronoiCornerSample(NoiseTileWrap(iv + offset.zxy, bTiling, RepeatSize), Quality);
+					mindist = min(mindist, dot(p, p));
+				}
+			}
+		}
+	}
+
+	// transform to -1 to 1 range as expected by later OutputMin to OutputMax transform
+	return (sqrt(mindist) * 2 ) - 1;
+}
+
+
 // -------- Simplex method (faster in higher dimensions because less samples are used, uses gradient noise for quality) ---------
 // <Dimensions>D:<Normal>/<Simplex> 1D:2, 2D:4/3, 3D:8/4, 4D:16/5 
 
diff --git a/Engine/Shaders/SimpleElementPixelShader.usf b/Engine/Shaders/SimpleElementPixelShader.usf
--- a/Engine/Shaders/SimpleElementPixelShader.usf
+++ b/Engine/Shaders/SimpleElementPixelShader.usf
@@ -71,6 +71,10 @@ float3 SimpleElementFrameBufferBlendOp(float4 Source)
 
 #if SE_BLEND_MODE == SE_BLEND_OPAQUE || SE_BLEND_MODE == SE_BLEND_MASKED || SE_BLEND_MODE == SE_BLEND_MASKEDDISTANCEFIELD || SE_BLEND_MODE == SE_BLEND_MASKEDDISTANCEFIELDSHADOWED
 	return Source.rgb;
+// AlphaComposite will set both MATERIALBLENDING_TRANSLUCENT and MATERIALBLENDING_ALPHACOMPOSITE defines
+// so ensure  MATERIALBLENDING_ALPHACOMPOSITE gets first in line
+#elif MATERIALBLENDING_ALPHACOMPOSITE
+	return Source.rgb + (Dest.rgb*(1.0 - Source.a));
 #elif SE_BLEND_MODE == SE_BLEND_TRANSLUCENT || SE_BLEND_MODE == SE_BLEND_TRANSLUCENTALPHAONLY || SE_BLEND_MODE == SE_BLEND_ALPHABLEND || SE_BLEND_MODE == SE_BLEND_TRANSLUCENTDISTANCEFIELD || SE_BLEND_MODE == SE_BLEND_TRANSLUCENTDISTANCEFIELDSHADOWED
 	return (Source.rgb*Source.a) + (Dest.rgb*(1.0 - Source.a));
 #elif SE_BLEND_MODE == SE_BLEND_ADDITIVE
diff --git a/Engine/Shaders/VulkanCommon.usf b/Engine/Shaders/VulkanCommon.usf
--- a/Engine/Shaders/VulkanCommon.usf
+++ b/Engine/Shaders/VulkanCommon.usf
@@ -5,6 +5,6 @@
 =============================================================================*/
 
 // Update this GUID to improve shader recompilation for Vulkan only shaders
-// GUID = 864401E6-0B37-40F0-850F-E5B79A6FEF56
+// GUID = 864401E6-0B37-40F0-850F-E5B79A6FEF68
 
 #pragma once
diff --git a/Engine/Source/Developer/MaterialUtilities/Private/MaterialUtilities.cpp b/Engine/Source/Developer/MaterialUtilities/Private/MaterialUtilities.cpp
--- a/Engine/Source/Developer/MaterialUtilities/Private/MaterialUtilities.cpp
+++ b/Engine/Source/Developer/MaterialUtilities/Private/MaterialUtilities.cpp
@@ -1751,7 +1751,98 @@ static bool GetUniformScale(const TArray<float> Scales, float& UniformScale)
 	return false;
 }
 
-bool FMaterialUtilities::ExportMaterialTexCoordScales(UMaterialInterface* InMaterial, EMaterialQualityLevel::Type QualityLevel, ERHIFeatureLevel::Type FeatureLevel, TArray<FMaterialTexCoordBuildInfo>& OutScales)
+uint32 GetTypeHash(const FMaterialUtilities::FExportErrorManager::FError& Error)
+{
+	return GetTypeHash(Error.Material);
+}
+
+bool FMaterialUtilities::FExportErrorManager::FError::operator==(const FError& Rhs) const
+{
+	return Material == Rhs.Material && RegisterIndex == Rhs.RegisterIndex && ErrorType == Rhs.ErrorType;
+}
+
+
+void FMaterialUtilities::FExportErrorManager::Register(const UMaterialInterface* Material, const UTexture* Texture, int32 RegisterIndex, EErrorType ErrorType)
+{
+	if (!Material || !Texture) return;
+
+	FError Error;
+	Error.Material = Material->GetMaterialResource(FeatureLevel);
+	if (!Error.Material) return;
+	Error.RegisterIndex = RegisterIndex;
+	Error.ErrorType = ErrorType;
+
+	FInstance Instance;
+	Instance.Material = Material;
+	Instance.Texture = Texture;
+
+	ErrorInstances.FindOrAdd(Error).Push(Instance);
+}
+
+void FMaterialUtilities::FExportErrorManager::OutputToLog()
+{
+	const UMaterialInterface* CurrentMaterial = nullptr;
+	int32 MaxInstanceCount = 0;
+	FString TextureErrors;
+
+	for (TMap<FError, TArray<FInstance> >::TIterator It(ErrorInstances);; ++It)
+	{
+		if (It && !It->Value.Num()) continue;
+
+		// Here we pack texture list per material.
+		if (!It || CurrentMaterial != It->Value[0].Material)
+		{
+			// Flush
+			if (CurrentMaterial)
+			{
+				FString SimilarCount(TEXT(""));
+				if (MaxInstanceCount > 1)
+				{
+					SimilarCount = FString::Printf(TEXT(", %d similar"), MaxInstanceCount - 1);
+				}
+
+				if (CurrentMaterial == CurrentMaterial->GetMaterial())
+				{
+					UE_LOG(LogMaterialUtilities, Warning, TEXT("Error generating scales for %s%s: %s"), *CurrentMaterial->GetName(), *SimilarCount, *TextureErrors);
+				}
+				else
+				{
+					UE_LOG(LogMaterialUtilities, Warning, TEXT("Error generating scales for %s, UMaterial=%s%s: %s"), *CurrentMaterial->GetName(), *CurrentMaterial->GetMaterial()->GetName(), *SimilarCount, *TextureErrors);
+				}
+			}
+
+			// Exit
+			if (!It)
+			{
+				break;
+			}
+
+			// Start new
+			CurrentMaterial = It->Value[0].Material;
+			MaxInstanceCount = It->Value.Num();
+			TextureErrors.Empty();
+		}
+		else
+		{
+			// Append
+			MaxInstanceCount = FMath::Max<int32>(MaxInstanceCount, It->Value.Num());
+		}
+
+		const TCHAR* ErrorMsg = TEXT("Unkown Error");
+		if (It->Key.ErrorType == EET_IncohorentValues)
+		{
+			ErrorMsg = TEXT("Incoherent");
+		}
+		else if (It->Key.ErrorType == EET_NoValues)
+		{
+			ErrorMsg = TEXT("NoValues");
+		}
+
+		TextureErrors.Append(FString::Printf(TEXT("(%s:%d,%s) "), ErrorMsg, It->Key.RegisterIndex, *It->Value[0].Texture->GetName()));
+	}
+}
+
+bool FMaterialUtilities::ExportMaterialTexCoordScales(UMaterialInterface* InMaterial, EMaterialQualityLevel::Type QualityLevel, ERHIFeatureLevel::Type FeatureLevel, TArray<FMaterialTexCoordBuildInfo>& OutScales, FExportErrorManager& OutErrors)
 {
 	TArray<FFloat16Color> RenderedVectors;
 
@@ -1905,21 +1996,7 @@ bool FMaterialUtilities::ExportMaterialTexCoordScales(UMaterialInterface* InMate
 				}
 			}
 
-			if (FailedTexture)
-			{
-				if (TextureIndexScales.Num())
-				{
-					// Cause 1 : the output does not map to an actual uniform scale.
-					// Cause 2 : the alogrithm fails to find the scale value (even though it is somewhat there).
-					UE_LOG(LogMaterialUtilities, Warning, TEXT("ExportMaterialTexCoordScales: Failed to find constant scale for texture index %d of material %s (bound to texture %s)"), RegisterIndex, *InMaterial->GetName(), *FailedTexture->GetName());
-				}
-				else
-				{
-					// Cause 1: the shader did not use this texture at all, resulting in value being the default value (could be because of some branching)
-					// Cause 2: the shader outputs the scale, but the scale was 0. That means the coordinate did not change between the pixels. For instance, if the mapping was based on the world position.
-					UE_LOG(LogMaterialUtilities, Warning, TEXT("ExportMaterialTexCoordScales: Failed to generate scales for texture index %d of material %s (bound to texture %s)"), RegisterIndex, *InMaterial->GetName(), *FailedTexture->GetName());
-				}
-			}
+			OutErrors.Register(InMaterial, FailedTexture, RegisterIndex, TextureIndexScales.Num() ? FExportErrorManager::EErrorType::EET_IncohorentValues : FExportErrorManager::EErrorType::EET_NoValues);
 		}
 	}
 
diff --git a/Engine/Source/Developer/MaterialUtilities/Public/MaterialUtilities.h b/Engine/Source/Developer/MaterialUtilities/Public/MaterialUtilities.h
--- a/Engine/Source/Developer/MaterialUtilities/Public/MaterialUtilities.h
+++ b/Engine/Source/Developer/MaterialUtilities/Public/MaterialUtilities.h
@@ -378,6 +378,61 @@ public:
 	*/
 	static void ResizeFlattenMaterial(FFlattenMaterial& InFlattenMaterial, const struct FMeshProxySettings& InMeshProxySettings);
 
+
+	/**
+	* Contains errors generated when exporting material texcoord scales. 
+	* Used to prevent displaying duplicates, as instances using the same shaders get the same issues.
+	*/
+	class MATERIALUTILITIES_API FExportErrorManager
+	{
+	public:
+
+		FExportErrorManager(ERHIFeatureLevel::Type InFeatureLevel) : FeatureLevel(InFeatureLevel) {}
+
+		enum EErrorType
+		{
+			EET_IncohorentValues,
+			EET_NoValues
+		};
+
+		/**
+		* Register a new error.
+		*
+		* @param Material			The material having this error.
+		* @param Texture			The texture for which the scale could not be generated.
+		* @param RegisterIndex		The register index bound to this texture.
+		* @param ErrorType			The issue encountered.
+		*/
+		void Register(const UMaterialInterface* Material, const UTexture* Texture, int32 RegisterIndex, EErrorType ErrorType);
+
+		/**
+		* Output all errors registered.
+		*/
+		void OutputToLog();
+
+	private:
+
+		struct FError
+		{
+			const FMaterial* Material;
+			int32 RegisterIndex;
+			EErrorType ErrorType;
+
+			bool operator==(const FError& Rhs) const;
+		};
+
+		struct FInstance
+		{
+			const UMaterialInterface* Material;
+			const UTexture* Texture;
+		};
+
+		friend uint32 GetTypeHash(const FError& Error);
+
+		ERHIFeatureLevel::Type FeatureLevel;
+		TMap<FError, TArray<FInstance> > ErrorInstances;
+	};
+
 	/**
 	* Get the material texcoord scales applied on each textures
 	*
@@ -387,7 +442,7 @@ public:
 	* @param OutScales			TheOutput array of rendered samples	
 	* @return					Whether operation was successful
 	*/
-	static bool ExportMaterialTexCoordScales(UMaterialInterface* InMaterial, EMaterialQualityLevel::Type QualityLevel, ERHIFeatureLevel::Type FeatureLevel, TArray<FMaterialTexCoordBuildInfo>& OutScales);
+	static bool ExportMaterialTexCoordScales(UMaterialInterface* InMaterial, EMaterialQualityLevel::Type QualityLevel, ERHIFeatureLevel::Type FeatureLevel, TArray<FMaterialTexCoordBuildInfo>& OutScales, FExportErrorManager& OutErrors);
 
 private:
 	
diff --git a/Engine/Source/Developer/MeshUtilities/Private/MeshUtilities.cpp b/Engine/Source/Developer/MeshUtilities/Private/MeshUtilities.cpp
--- a/Engine/Source/Developer/MeshUtilities/Private/MeshUtilities.cpp
+++ b/Engine/Source/Developer/MeshUtilities/Private/MeshUtilities.cpp
@@ -5080,8 +5080,8 @@ bool FMeshUtilities::BuildSkeletalMesh_Legacy(FStaticLODModel& LODModel, const F
 	SkeletalMeshTools::BuildSkeletalMeshChunks(Faces, RawVertices, VertIndexAndZ, bKeepOverlappingVertices, Chunks, bTooManyVerts);
 
 	// Chunk vertices to satisfy the requested limit.
-	static const auto MaxBonesVar = IConsoleManager::Get().FindTConsoleVariableDataInt(TEXT("Compat.MAX_GPUSKIN_BONES"));
-	const int32 MaxGPUSkinBones = MaxBonesVar->GetValueOnAnyThread();
+	const uint32 MaxGPUSkinBones = FGPUBaseSkinVertexFactory::GetMaxGPUSkinBones();
+	check(MaxGPUSkinBones <= FGPUBaseSkinVertexFactory::GHardwareMaxGPUSkinBones);
 	SkeletalMeshTools::ChunkSkinnedVertices(Chunks, MaxGPUSkinBones);
 
 	// Build the skeletal model from chunks.
diff --git a/Engine/Source/Developer/VulkanShaderFormat/Private/VulkanBackend.cpp b/Engine/Source/Developer/VulkanShaderFormat/Private/VulkanBackend.cpp
--- a/Engine/Source/Developer/VulkanShaderFormat/Private/VulkanBackend.cpp
+++ b/Engine/Source/Developer/VulkanShaderFormat/Private/VulkanBackend.cpp
@@ -36,6 +36,7 @@
 #include "hlslcc_private.h"
 #include "VulkanBackend.h"
 #include "compiler.h"
+#include "ShaderCompilerCommon.h"
 
 #include "VulkanConfiguration.h"
 
@@ -49,6 +50,7 @@ PRAGMA_ENABLE_SHADOW_VARIABLE_WARNINGS
 #include "IRDump.h"
 //@todo-rco: Remove STL!
 #include <sstream>
+#include <vector>
 //#define OPTIMIZE_ANON_STRUCTURES_OUT
 // We can't optimize them out presently, because apparently Windows Radeon
 // OpenGL driver chokes on valid GLSL code then.
@@ -57,7 +59,7 @@ PRAGMA_ENABLE_SHADOW_VARIABLE_WARNINGS
 #define _strdup strdup
 #endif
 
-static inline std::string FixHlslName(const glsl_type* Type)
+static inline std::string FixHlslName(const glsl_type* Type, bool bUseTextureInsteadOfSampler = false)
 {
 	check(Type->is_image() || Type->is_vector() || Type->is_numeric() || Type->is_void() || Type->is_sampler() || Type->is_scalar());
 	std::string Name = Type->name;
@@ -113,6 +115,25 @@ static inline std::string FixHlslName(const glsl_type* Type)
 	{
 		return "mat4";
 	}
+	else if (Type->is_sampler() && !Type->sampler_buffer && bUseTextureInsteadOfSampler)
+	{
+		if (!strcmp(Type->HlslName, "texturecube"))
+		{
+			return "textureCube";
+		}
+		else if (!strcmp(Type->HlslName, "texture2d"))
+		{
+			return "texture2D";
+		}
+		else if (!strcmp(Type->HlslName, "texture3d"))
+		{
+			return "texture3D";
+		}
+		else
+		{
+			return Type->HlslName;
+		}
+	}
 
 	return Name;
 }
@@ -497,7 +518,7 @@ static inline EDescriptorSetStage GetDescriptorSetForStage(_mesa_glsl_parser_tar
 /**
 * IR visitor used to generate GLSL. Based on ir_print_visitor.
 */
-class vulkan_ir_gen_glsl_visitor : public ir_visitor
+class FGenerateVulkanVisitor : public ir_visitor
 {
 	/** Track which multi-dimensional arrays are used. */
 	struct md_array_entry : public exec_node
@@ -552,7 +573,7 @@ class vulkan_ir_gen_glsl_visitor : public ir_visitor
 	bool bIsES31;
 	EHlslCompileTarget Target;
 	_mesa_glsl_parser_targets ShaderTarget;
-
+	FVulkanLanguageSpec* LanguageSpec;
 	bool bGenerateLayoutLocations;
 	bool bDefaultPrecisionIsHalf;
 
@@ -587,6 +608,8 @@ class vulkan_ir_gen_glsl_visitor : public ir_visitor
 	// Found dFdx or dFdy
 	bool bUsesDXDY;
 
+	std::vector<std::string> SamplerStateNames;
+
 	/**
 	* Return true if the type is a multi-dimensional array. Also, track the
 	* array.
@@ -714,7 +737,7 @@ class vulkan_ir_gen_glsl_visitor : public ir_visitor
 		}
 		else
 		{
-			std::string Name = FixHlslName(t);
+			std::string Name = FixHlslName(t, LanguageSpec->AllowsSharingSamplers());
 			ralloc_asprintf_append(buffer, "%s", Name.c_str());
 		}
 	}
@@ -819,18 +842,18 @@ class vulkan_ir_gen_glsl_visitor : public ir_visitor
 		return GLSL_PRECISION_DEFAULT;
 	}
 
-	void AppendPrecisionModifier(char** inBuffer, EPrecisionModifier PrecisionModifier)
+	const char* GetPrecisionModifierName(EPrecisionModifier PrecisionModifier)
 	{
 		switch (PrecisionModifier)
 		{
 		case GLSL_PRECISION_LOWP:
-			ralloc_asprintf_append(inBuffer, "lowp ");
+			return "lowp";
 			break;
 		case GLSL_PRECISION_MEDIUMP:
-			ralloc_asprintf_append(inBuffer, "mediump ");
+			return "mediump";
 			break;
 		case GLSL_PRECISION_HIGHP:
-			ralloc_asprintf_append(inBuffer, "highp ");
+			return "highp";
 			break;
 		case GLSL_PRECISION_DEFAULT:
 			break;
@@ -838,6 +861,12 @@ class vulkan_ir_gen_glsl_visitor : public ir_visitor
 			// we missed a type
 			check(false);
 		}
+		return "";
+	}
+
+	inline void AppendPrecisionModifier(char** inBuffer, EPrecisionModifier PrecisionModifier)
+	{
+		ralloc_asprintf_append(inBuffer, "%s ", GetPrecisionModifierName(PrecisionModifier));
 	}
 
 	/**
@@ -1057,21 +1086,25 @@ class vulkan_ir_gen_glsl_visitor : public ir_visitor
 			}
 			else
 			{
-
 				char* layout = nullptr;
 
+				uint32 Interpolation = var->interpolation;
 				if (var->type->is_sampler())
 				{
 					layout = ralloc_asprintf(nullptr,
 						"layout(set=%d, binding=%d) ",
 						GetDescriptorSetForStage(ShaderTarget),
-						BindingTable.RegisterBinding(var->name, "s", var->type->sampler_buffer ? FVulkanBindingTable::TYPE_SAMPLER_BUFFER : FVulkanBindingTable::TYPE_SAMPLER));
+						BindingTable.RegisterBinding(var->name, "s", var->type->sampler_buffer ? FVulkanBindingTable::TYPE_SAMPLER_BUFFER : FVulkanBindingTable::TYPE_COMBINED_IMAGE_SAMPLER));
 				}
-
-				if (bGenerateLayoutLocations && var->explicit_location)
+				else if (bGenerateLayoutLocations && var->explicit_location)
 				{
 					check(layout_bits == 0);
 					layout = ralloc_asprintf(nullptr, "layout(location=%d) ", var->location);
+					if (ShaderTarget == fragment_shader && var->type->is_integer() && var->mode == ir_var_in)
+					{
+						// Flat
+						Interpolation = 2;
+					}
 				}
 
 				ralloc_asprintf_append(
@@ -1082,7 +1115,7 @@ class vulkan_ir_gen_glsl_visitor : public ir_visitor
 					invariant_str[var->invariant],
 					patch_constant_str[var->is_patch_constant],
 					mode_str[var->mode],
-					interp_str[var->interpolation]
+					interp_str[Interpolation]
 					);
 
 				if (bEmitPrecision)
@@ -1120,7 +1153,6 @@ class vulkan_ir_gen_glsl_visitor : public ir_visitor
 		// this is for the case of a variable that is declared, but not later dereferenced (which can happen
 		// when debugging HLSLCC and running without optimization
 		AddTypeToUsedStructs(var->type);
-
 	}
 
 	virtual void visit(ir_function_signature *sig)
@@ -1324,7 +1356,36 @@ class vulkan_ir_gen_glsl_visitor : public ir_visitor
 			offset_str[tex->offset != 0],
 			EXT_str[(int)bEmitEXT]
 			);
-		tex->sampler->accept(this);
+
+		if (LanguageSpec->AllowsSharingSamplers() && !tex->sampler->type->sampler_buffer && tex->op != ir_txf)
+		{
+			uint32 SSIndex = AddUniqueSamplerState(tex->SamplerStateName);
+			char PackedName[256];
+			sprintf_s(PackedName, "%sz%d", glsl_variable_tag_from_parser_target(ShaderTarget), SSIndex);
+			BindingTable.RegisterBinding(PackedName, "z", FVulkanBindingTable::TYPE_SAMPLER);
+
+			auto GetSamplerSuffix = [](int32 Dim)
+			{
+				switch (Dim)
+				{
+				case GLSL_SAMPLER_DIM_1D: return "1D";
+				case GLSL_SAMPLER_DIM_2D: return "2D";
+				case GLSL_SAMPLER_DIM_3D: return "3D";
+				case GLSL_SAMPLER_DIM_CUBE: return "Cube";
+				//case GLSL_SAMPLER_DIM_RECT: return "Rect";
+				//case GLSL_SAMPLER_DIM_BUF: return "Buf";
+				default: return "INVALID";
+				}
+			};
+
+			ralloc_asprintf_append(buffer, "sampler%s(", GetSamplerSuffix(tex->sampler->type->sampler_dimensionality));
+			tex->sampler->accept(this);
+			ralloc_asprintf_append(buffer, ", %s)", PackedName);
+		}
+		else
+		{
+			tex->sampler->accept(this);
+		}
 
 		// Emit coordinates.
 		if ((op == ir_txs && tex->lod_info.lod) || op == ir_txm)
@@ -2309,7 +2370,7 @@ class vulkan_ir_gen_glsl_visitor : public ir_visitor
 				int32 Binding = BindingTable.RegisterBinding(block_name, var_name, Type);
 				ralloc_asprintf_append(
 					buffer,
-					"layout(set=%d, binding = %d, std140) uniform %s\n{\n",
+					"layout(set=%d, binding=%d, std140) uniform %s\n{\n",
 					GetDescriptorSetForStage(ShaderTarget),
 					Binding,
 					block_name);
@@ -2797,6 +2858,16 @@ class vulkan_ir_gen_glsl_visitor : public ir_visitor
 				ralloc_asprintf_append(buffer, "\n");
 			}
 		}
+
+		if (!SamplerStateNames.empty())
+		{
+			ralloc_asprintf_append(buffer, "// @SamplerStates: ");
+			for (uint32 Index = 0; Index < SamplerStateNames.size(); ++Index)
+			{
+				ralloc_asprintf_append(buffer, "%s%d:%s", Index > 0 ? "," : "", Index, SamplerStateNames[Index].c_str());
+			}
+			ralloc_asprintf_append(buffer, "\n");
+		}
 	}
 
 	/**
@@ -2939,9 +3010,10 @@ class vulkan_ir_gen_glsl_visitor : public ir_visitor
 public:
 
 	/** Constructor. */
-	vulkan_ir_gen_glsl_visitor(EHlslCompileTarget InTarget,
+	FGenerateVulkanVisitor(EHlslCompileTarget InTarget,
 							FVulkanBindingTable& InBindingTable,
 							_mesa_glsl_parser_targets InShaderTarget,
+							FVulkanLanguageSpec* InLanguageSpec,
 							bool bInGenerateLayoutLocations,
 							bool bInDefaultPrecisionIsHalf)
 		: early_depth_stencil(false)
@@ -2950,6 +3022,7 @@ public:
 		, wg_size_z(0)
 		, Target(InTarget)
 		, ShaderTarget(InShaderTarget)
+		, LanguageSpec(InLanguageSpec)
 		, bGenerateLayoutLocations(bInGenerateLayoutLocations)
 		, bDefaultPrecisionIsHalf(bInDefaultPrecisionIsHalf)
 		, BindingTable(InBindingTable)
@@ -2974,13 +3047,27 @@ public:
 	}
 
 	/** Destructor. */
-	virtual ~vulkan_ir_gen_glsl_visitor()
+	virtual ~FGenerateVulkanVisitor()
 	{
 		hash_table_dtor(printable_names);
 		hash_table_dtor(used_structures);
 		hash_table_dtor(used_uniform_blocks);
 	}
 
+	int32 AddUniqueSamplerState(const std::string& Name)
+	{
+		for (uint32 Index = 0; Index < SamplerStateNames.size(); ++Index)
+		{
+			if (SamplerStateNames[Index] == Name)
+			{
+				return (int32)Index;
+			}
+		}
+
+		SamplerStateNames.push_back(Name);
+		return SamplerStateNames.size() - 1;
+	};
+
 	/**
 	* Executes the visitor on the provided ir.
 	* @returns the GLSL source code generated.
@@ -3000,10 +3087,12 @@ public:
 			ralloc_asprintf_append(buffer, "precision %s float;\n", DefaultPrecision);
 			ralloc_asprintf_append(buffer, "precision %s int;\n", DefaultPrecision);
 			//ralloc_asprintf_append(buffer, "\n#ifndef DONTEMITSAMPLERDEFAULTPRECISION\n");
+			ralloc_asprintf_append(buffer, "precision %s sampler;\n", DefaultPrecision);
 			ralloc_asprintf_append(buffer, "precision %s sampler2D;\n", DefaultPrecision);
 			ralloc_asprintf_append(buffer, "precision %s samplerCube;\n", DefaultPrecision);
 			//ralloc_asprintf_append(buffer, "#endif\n");
 		}
+
 		// FramebufferFetchES2 'intrinsic'
 		bool bUsesFramebufferFetchES2 = false;//UsesUEIntrinsic(ir, FRAMEBUFFER_FETCH_ES2);
 		/*
@@ -3126,6 +3215,55 @@ public:
 				break;
 			}
 		}
+
+		// Here since the code_buffer must have been populated beforehand
+		if (LanguageSpec->AllowsSharingSamplers())
+		{
+			auto FindPrecision = [&](int32 Index)
+			{
+				for (auto& Pair : state->TextureToSamplerMap)
+				{
+					for (auto& Entry : Pair.second)
+					{
+						if (Entry == SamplerStateNames[Index])
+						{
+							for (auto& PackedEntry : state->GlobalPackedArraysMap['s'])
+							{
+								if (!strcmp(Pair.first.c_str(), PackedEntry.Name.c_str()))
+								{
+									foreach_iter(exec_list_iterator, iter, sampler_variables)
+									{
+										ir_variable* var = ((extern_var*)iter.get())->var;
+										if (!strcmp(var->name, PackedEntry.CB_PackedSampler.c_str()))
+										{
+											return GetPrecisionModifierName(GetPrecisionModifier(var->type));
+										}
+									}
+								}
+							}
+						}
+					}
+				}
+
+				return "";
+			};
+
+			const auto& Bindings = BindingTable.GetBindings();
+			for (int32 Index = 0; Index < Bindings.Num(); ++Index)
+			{
+				if (Bindings[Index].Type == FVulkanBindingTable::TYPE_SAMPLER)
+				{
+					int32 Binding = atoi(Bindings[Index].Name + 2);
+					const char* Precision = FindPrecision(Binding);
+
+					ralloc_asprintf_append(buffer, "layout(set=%d, binding=%d) uniform %s sampler %sz%d;\n",
+						GetDescriptorSetForStage(ShaderTarget), Index,
+						Precision,
+						glsl_variable_tag_from_parser_target(ShaderTarget), Binding);
+				}
+			}
+		}
+
 		buffer = 0;
 
 		static const char* vulkan_required_extensions =
@@ -3230,7 +3368,7 @@ struct FBreakPrecisionChangesVisitor : public ir_rvalue_visitor
 	}
 };
 
-void vulkan_ir_gen_glsl_visitor::AddTypeToUsedStructs(const glsl_type* type)
+void FGenerateVulkanVisitor::AddTypeToUsedStructs(const glsl_type* type)
 {
 	if (type->base_type == GLSL_TYPE_STRUCT)
 	{
@@ -3274,7 +3412,7 @@ char* FVulkanCodeBackend::GenerateCode(exec_list* ir, _mesa_glsl_parse_state* st
 	const bool bCanHaveUBs = true;//(HlslCompileFlags & HLSLCC_FlattenUniformBuffers) != HLSLCC_FlattenUniformBuffers;
 
 	// Setup root visitor
-	vulkan_ir_gen_glsl_visitor visitor(Target, BindingTable, state->target, bGenerateLayoutLocations, bDefaultPrecisionIsHalf);
+	FGenerateVulkanVisitor visitor(Target, BindingTable, state->target, (FVulkanLanguageSpec*)state->LanguageSpec, bGenerateLayoutLocations, bDefaultPrecisionIsHalf);
 
 	const char* code = visitor.run(ir, state, bGroupFlattenedUBs, bCanHaveUBs);
 
@@ -5269,14 +5407,15 @@ void FVulkanCodeBackend::GenShaderPatchConstantFunctionInputs(_mesa_glsl_parse_s
 
 void FVulkanLanguageSpec::SetupLanguageIntrinsics(_mesa_glsl_parse_state* State, exec_list* ir)
 {
+/*
 	if (bIsES2)
 	{
 		make_intrinsic_genType(ir, State, FRAMEBUFFER_FETCH_ES2, ir_invalid_opcode, IR_INTRINSIC_FLOAT, 0, 4, 4);
 		make_intrinsic_genType(ir, State, DEPTHBUFFER_FETCH_ES2, ir_invalid_opcode, IR_INTRINSIC_ALL_FLOATING, 3, 1, 1);
 		make_intrinsic_genType(ir, State, GET_HDR_32BPP_HDR_ENCODE_MODE_ES2, ir_invalid_opcode, IR_INTRINSIC_ALL_FLOATING, 0);
 	}
-
-	if (State->language_version >= 310)
+*/
+	//if (State->language_version >= 310)
 	{
 		/**
 		* Create GLSL functions that are left out of the symbol table
@@ -5368,7 +5507,7 @@ FVulkanBindingTable::FBinding::FBinding(const char* InName, int32 InIndex, EBind
 	FMemory::Memcpy(Name, InName, NewNameLength);
 
 	// Validate Sampler type, s == PACKED_TYPENAME_SAMPLER
-	check((Type == TYPE_SAMPLER || Type == TYPE_SAMPLER_BUFFER) ? SubType == 's' : true);
+	check((Type == TYPE_COMBINED_IMAGE_SAMPLER || Type == TYPE_SAMPLER_BUFFER) ? SubType == 's' : true);
 
 	check(Type == TYPE_PACKED_UNIFORM_BUFFER ?
 		( SubType == 'h' || SubType == 'm' || SubType == 'l' || SubType == 'i' || SubType == 'u' ) : true);
diff --git a/Engine/Source/Developer/VulkanShaderFormat/Private/VulkanBackend.h b/Engine/Source/Developer/VulkanShaderFormat/Private/VulkanBackend.h
--- a/Engine/Source/Developer/VulkanShaderFormat/Private/VulkanBackend.h
+++ b/Engine/Source/Developer/VulkanShaderFormat/Private/VulkanBackend.h
@@ -8,10 +8,12 @@
 class FVulkanLanguageSpec : public ILanguageSpec
 {
 protected:
-	bool bIsES2;
+	bool bShareSamplers;
 
 public:
-	FVulkanLanguageSpec(bool bInIsES2) : bIsES2(bInIsES2) {}
+	FVulkanLanguageSpec(bool bInShareSamplers)
+		: bShareSamplers(bInShareSamplers)
+	{}
 
 	virtual bool SupportsDeterminantIntrinsic() const override
 	{
@@ -32,13 +34,12 @@ public:
 
 	virtual void SetupLanguageIntrinsics(_mesa_glsl_parse_state* State, exec_list* ir) override;
 
-	//#todo-rco: Enable
-	virtual bool AllowsSharingSamplers() const override { return false; }
+	virtual bool AllowsSharingSamplers() const override { return bShareSamplers; }
 };
 
 class ir_variable;
 
-// Generates GLSL compliant code from IR tokens
+// Generates Vulkan compliant code from IR tokens
 #ifdef __GNUC__
 #pragma GCC visibility push(default)
 #endif // __GNUC__
@@ -47,10 +48,12 @@ struct FVulkanBindingTable
 {
 	enum EBindingType : uint16
 	{
-		TYPE_SAMPLER,
+		TYPE_COMBINED_IMAGE_SAMPLER,
 		TYPE_SAMPLER_BUFFER,
 		TYPE_UNIFORM_BUFFER,
 		TYPE_PACKED_UNIFORM_BUFFER,
+		TYPE_SAMPLER,
+		TYPE_IMAGE,
 
 		TYPE_MAX,
 	};
diff --git a/Engine/Source/Developer/VulkanShaderFormat/Private/VulkanShaderCompiler.cpp b/Engine/Source/Developer/VulkanShaderFormat/Private/VulkanShaderCompiler.cpp
--- a/Engine/Source/Developer/VulkanShaderFormat/Private/VulkanShaderCompiler.cpp
+++ b/Engine/Source/Developer/VulkanShaderFormat/Private/VulkanShaderCompiler.cpp
@@ -12,15 +12,15 @@
 
 DEFINE_LOG_CATEGORY_STATIC(LogVulkanShaderCompiler, Log, All); 
 
-static int32 GUseExternalShaderCompiler = 0;
-static FAutoConsoleVariableRef CVarVulkanUseExternalShaderCompiler(
-	TEXT("r.Vulkan.UseExternalShaderCompiler"),
-	GUseExternalShaderCompiler,
-	TEXT("Whether to use the internal shader compiling library or the external glslang tool.\n")
-	TEXT(" 0: Internal compiler\n")
-	TEXT(" 1: External compiler)"),
-	ECVF_Default
-	);
+//static int32 GUseExternalShaderCompiler = 0;
+//static FAutoConsoleVariableRef CVarVulkanUseExternalShaderCompiler(
+//	TEXT("r.Vulkan.UseExternalShaderCompiler"),
+//	GUseExternalShaderCompiler,
+//	TEXT("Whether to use the internal shader compiling library or the external glslang tool.\n")
+//	TEXT(" 0: Internal compiler\n")
+//	TEXT(" 1: External compiler)"),
+//	ECVF_Default
+//	);
 
 extern bool GenerateSpirv(const ANSICHAR* Source, FCompilerInfo& CompilerInfo, FString& OutErrors, const FString& DumpDebugInfoPath, TArray<uint8>& OutSpirv);
 
@@ -206,7 +206,7 @@ static inline FString GetExtension(EHlslShaderFrequency Frequency, bool bAddDot
 	case HSF_DomainShader:		Name = TEXT(".tese"); break;
 	}
 
-	if (bAddDot)
+	if (!bAddDot)
 	{
 		++Name;
 	}
@@ -239,15 +239,15 @@ static uint32 GetTypeComponents(const FString& Type)
 
 static void GenerateBindingTable(const FVulkanShaderSerializedBindings& SerializedBindings, FVulkanShaderBindingTable& OutBindingTable)
 {
-	int32 NumSamplers = 0;
+	int32 NumCombinedSamplers = 0;
 	int32 NumSamplerBuffers = 0;
 	int32 NumUniformBuffers = 0;
 
 	auto& Layouts = SerializedBindings.Bindings;
 
 	//#todo-rco: FIX! SamplerBuffers share numbering with Samplers
-	NumSamplers = Layouts[FVulkanShaderSerializedBindings::TYPE_SAMPLER].Num() + Layouts[FVulkanShaderSerializedBindings::TYPE_SAMPLER_BUFFER].Num();
-	NumSamplerBuffers = Layouts[FVulkanShaderSerializedBindings::TYPE_SAMPLER].Num() + Layouts[FVulkanShaderSerializedBindings::TYPE_SAMPLER_BUFFER].Num();
+	NumCombinedSamplers = Layouts[FVulkanShaderSerializedBindings::TYPE_COMBINED_IMAGE_SAMPLER].Num() + Layouts[FVulkanShaderSerializedBindings::TYPE_SAMPLER_BUFFER].Num();
+	NumSamplerBuffers = Layouts[FVulkanShaderSerializedBindings::TYPE_COMBINED_IMAGE_SAMPLER].Num() + Layouts[FVulkanShaderSerializedBindings::TYPE_SAMPLER_BUFFER].Num();
 	NumUniformBuffers = Layouts[FVulkanShaderSerializedBindings::TYPE_PACKED_UNIFORM_BUFFER].Num() + Layouts[FVulkanShaderSerializedBindings::TYPE_UNIFORM_BUFFER].Num();
 
 	for (int32 Index = 0; Index < CrossCompiler::PACKED_TYPEINDEX_MAX; ++Index)
@@ -255,15 +255,15 @@ static void GenerateBindingTable(const FVulkanShaderSerializedBindings& Serializ
 		OutBindingTable.PackedGlobalUBsIndices[Index] = -1;
 	}
 
-	OutBindingTable.SamplerBindingIndices.AddUninitialized(NumSamplers);
+	OutBindingTable.CombinedSamplerBindingIndices.AddUninitialized(NumCombinedSamplers);
 	//#todo-rco: FIX! SamplerBuffers share numbering with Samplers
 	OutBindingTable.SamplerBufferBindingIndices.AddUninitialized(NumSamplerBuffers);
 	OutBindingTable.UniformBufferBindingIndices.AddUninitialized(NumUniformBuffers);
 
-	for (int32 Index = 0; Index < Layouts[FVulkanShaderSerializedBindings::TYPE_SAMPLER].Num(); ++Index)
+	for (int32 Index = 0; Index < Layouts[FVulkanShaderSerializedBindings::TYPE_COMBINED_IMAGE_SAMPLER].Num(); ++Index)
 	{
-		auto& Mapping = Layouts[FVulkanShaderSerializedBindings::TYPE_SAMPLER][Index];
-		OutBindingTable.SamplerBindingIndices[Mapping.EngineBindingIndex] = Mapping.VulkanBindingIndex;
+		auto& Mapping = Layouts[FVulkanShaderSerializedBindings::TYPE_COMBINED_IMAGE_SAMPLER][Index];
+		OutBindingTable.CombinedSamplerBindingIndices[Mapping.EngineBindingIndex] = Mapping.VulkanBindingIndex;
 		//#todo-rco: FIX! SamplerBuffers share numbering with Samplers
 		OutBindingTable.SamplerBufferBindingIndices[Mapping.EngineBindingIndex] = Mapping.VulkanBindingIndex;
 	}
@@ -271,7 +271,7 @@ static void GenerateBindingTable(const FVulkanShaderSerializedBindings& Serializ
 	for (int32 Index = 0; Index < Layouts[FVulkanShaderSerializedBindings::TYPE_SAMPLER_BUFFER].Num(); ++Index)
 	{
 		auto& Mapping = Layouts[FVulkanShaderSerializedBindings::TYPE_SAMPLER_BUFFER][Index];
-		OutBindingTable.SamplerBindingIndices[Mapping.EngineBindingIndex] = Mapping.VulkanBindingIndex;
+		OutBindingTable.CombinedSamplerBindingIndices[Mapping.EngineBindingIndex] = Mapping.VulkanBindingIndex;
 		//#todo-rco: FIX! SamplerBuffers share numbering with Samplers
 		OutBindingTable.SamplerBufferBindingIndices[Mapping.EngineBindingIndex] = Mapping.VulkanBindingIndex;
 	}
@@ -292,7 +292,7 @@ static void GenerateBindingTable(const FVulkanShaderSerializedBindings& Serializ
 	}
 
 	// Do not share numbers here
-	OutBindingTable.NumDescriptorsWithoutPackedUniformBuffers = Layouts[FVulkanShaderSerializedBindings::TYPE_SAMPLER].Num() + Layouts[FVulkanShaderSerializedBindings::TYPE_SAMPLER_BUFFER].Num() + Layouts[FVulkanShaderSerializedBindings::TYPE_UNIFORM_BUFFER].Num();
+	OutBindingTable.NumDescriptorsWithoutPackedUniformBuffers = Layouts[FVulkanShaderSerializedBindings::TYPE_COMBINED_IMAGE_SAMPLER].Num() + Layouts[FVulkanShaderSerializedBindings::TYPE_SAMPLER_BUFFER].Num() + Layouts[FVulkanShaderSerializedBindings::TYPE_UNIFORM_BUFFER].Num();
 	OutBindingTable.NumDescriptors = OutBindingTable.NumDescriptorsWithoutPackedUniformBuffers + Layouts[FVulkanShaderSerializedBindings::TYPE_PACKED_UNIFORM_BUFFER].Num();
 }
 
@@ -415,8 +415,11 @@ static void BuildShaderOutput(
 		auto Type = FVulkanShaderSerializedBindings::TYPE_MAX;
 		switch (CurrBinding.Type)
 		{
-		case FVulkanBindingTable::TYPE_SAMPLER:
-			Type = FVulkanShaderSerializedBindings::TYPE_SAMPLER;
+		//case FVulkanBindingTable::TYPE_SAMPLER:
+		//	Type = FVulkanShaderSerializedBindings::TYPE_SAMPLER;
+		//	break;
+		case FVulkanBindingTable::TYPE_COMBINED_IMAGE_SAMPLER:
+			Type = FVulkanShaderSerializedBindings::TYPE_COMBINED_IMAGE_SAMPLER;
 			break;
 		case FVulkanBindingTable::TYPE_SAMPLER_BUFFER:
 			Type = FVulkanShaderSerializedBindings::TYPE_SAMPLER_BUFFER;
@@ -909,7 +912,7 @@ FCompilerInfo::FCompilerInfo(const FShaderCompilerInput& InInput, const FString&
 
 /**
  * Compile a shader using the external shader compiler
- */
+ *
 static void CompileUsingExternal(const struct FShaderCompilerInput& Input, struct FShaderCompilerOutput& Output, const class FString& WorkingDirectory, EVulkanShaderVersion Version)
 {
 	FString PreprocessedShader;
@@ -1033,7 +1036,7 @@ static void CompileUsingExternal(const struct FShaderCompilerInput& Input, struc
 		FVulkanBindingTable BindingTableES(Frequency);
 
 		FVulkanCodeBackend VulkanBackendES(CCFlagsES, BindingTableES, HlslCompilerTargetES);
-		FVulkanLanguageSpec VulkanLanguageSpec(false);
+		FVulkanLanguageSpec VulkanLanguageSpec(false, true);
 
 		int32 Result = 0;
 		if (!bIsSM5)
@@ -1271,8 +1274,8 @@ static bool CallHlslcc(const FString& PreprocessedShader, FVulkanBindingTable& B
 		// Call hlslcc
 		FVulkanCodeBackend VulkanBackend(CompilerInfo.CCFlags, BindingTable, HlslCompilerTarget);
 		FHlslCrossCompilerContext CrossCompilerContext(CompilerInfo.CCFlags, CompilerInfo.Frequency, HlslCompilerTarget);
-		//#todo-rco: Always false?
-		FVulkanLanguageSpec VulkanLanguageSpec(false);
+		const bool bShareSamplers = false;
+		FVulkanLanguageSpec VulkanLanguageSpec(bShareSamplers);
 		int32 Result = 0;
 		if (CrossCompilerContext.Init(TCHAR_TO_ANSI(*CompilerInfo.Input.SourceFilename), &VulkanLanguageSpec))
 		{
@@ -1343,12 +1346,12 @@ void CompileShader_Windows_Vulkan(const FShaderCompilerInput& Input, FShaderComp
 {
 	check(IsVulkanPlatform((EShaderPlatform)Input.Target.Platform));
 
-	if (GUseExternalShaderCompiler)
-	{
-		// Old path...
-		CompileUsingExternal(Input, Output, WorkingDirectory, Version);
-		return;
-	}
+	//if (GUseExternalShaderCompiler)
+	//{
+	//	// Old path...
+	//	CompileUsingExternal(Input, Output, WorkingDirectory, Version);
+	//	return;
+	//}
 
 	const bool bIsSM5 = (Version == EVulkanShaderVersion::SM5);
 	const bool bIsSM4 = (Version == EVulkanShaderVersion::SM4);
@@ -1505,8 +1508,9 @@ void CompileShader_Windows_Vulkan(const FShaderCompilerInput& Input, FShaderComp
 		//}
 		//else
 		{
-			// For debugging...
-			auto* Code = GeneratedGlslSource.GetData();
+			// For debugging; if you hit an error from Glslang/Spirv, use the SourceNoHeader for line numbers
+			auto* SourceWithHeader = GeneratedGlslSource.GetData();
+			char* SourceNoHeader = strstr(SourceWithHeader, "#version");
 			CompileUsingInternal(CompilerInfo, BindingTable, GeneratedGlslSource, EntryPointName, Output);
 		}
 	}
diff --git a/Engine/Source/Editor/MaterialEditor/Private/SMaterialEditorViewport.cpp b/Engine/Source/Editor/MaterialEditor/Private/SMaterialEditorViewport.cpp
--- a/Engine/Source/Editor/MaterialEditor/Private/SMaterialEditorViewport.cpp
+++ b/Engine/Source/Editor/MaterialEditor/Private/SMaterialEditorViewport.cpp
@@ -124,7 +124,7 @@ FLinearColor FMaterialEditorViewportClient::GetBackgroundColor() const
 			{
 				BackgroundColor = FLinearColor::White;
 			}
-			else if(PreviewBlendMode == BLEND_Translucent)
+			else if(PreviewBlendMode == BLEND_Translucent || PreviewBlendMode == BLEND_AlphaComposite)
 			{
 				BackgroundColor = FColor(64, 64, 64);
 			}
diff --git a/Engine/Source/Editor/NiagaraEditor/NiagaraEditor.Build.cs b/Engine/Source/Editor/NiagaraEditor/NiagaraEditor.Build.cs
--- a/Engine/Source/Editor/NiagaraEditor/NiagaraEditor.Build.cs
+++ b/Engine/Source/Editor/NiagaraEditor/NiagaraEditor.Build.cs
@@ -8,8 +8,9 @@ public class NiagaraEditor : ModuleRules
 	{
 		PrivateDependencyModuleNames.AddRange(
 			new string[] {
-				"Engine", 
-				"Core", 
+				"Engine",
+                "RHI",
+                "Core", 
 				"CoreUObject", 
                 "InputCore",
 				"RenderCore",
diff --git a/Engine/Source/Editor/NiagaraEditor/Private/NiagaraCompiler.cpp b/Engine/Source/Editor/NiagaraEditor/Private/NiagaraCompiler.cpp
--- a/Engine/Source/Editor/NiagaraEditor/Private/NiagaraCompiler.cpp
+++ b/Engine/Source/Editor/NiagaraEditor/Private/NiagaraCompiler.cpp
@@ -172,8 +172,16 @@ bool FNiagaraCompiler::CheckOutputs(FName OpName, TArray<TNiagaraExprPtr>& Outpu
 
 void FNiagaraCompiler::Error(FText ErrorText, UNiagaraNode* Node, UEdGraphPin* Pin)
 {
-	FString ErrorString = FString::Printf(TEXT("Node: @@ - Pin: @@ - %s"), *ErrorText.ToString());
-	MessageLog.Error(*ErrorString, Node, Pin);
+	if (Pin)
+	{
+		FString ErrorString = FString::Printf(TEXT("Node: @@ - Pin: @@ - %s"), *ErrorText.ToString());
+		MessageLog.Error(*ErrorString, Node, Pin);
+	}
+	else
+	{
+		FString ErrorString = FString::Printf(TEXT("Node: @@ - %s"), *ErrorText.ToString());
+		MessageLog.Error(*ErrorString, Node);
+	}
 }
 
 void FNiagaraCompiler::Warning(FText WarningText, UNiagaraNode* Node, UEdGraphPin* Pin)
diff --git a/Engine/Source/Editor/NiagaraEditor/Private/NiagaraEmitterPropertiesDetailsCustomization.cpp b/Engine/Source/Editor/NiagaraEditor/Private/NiagaraEmitterPropertiesDetailsCustomization.cpp
--- a/Engine/Source/Editor/NiagaraEditor/Private/NiagaraEmitterPropertiesDetailsCustomization.cpp
+++ b/Engine/Source/Editor/NiagaraEditor/Private/NiagaraEmitterPropertiesDetailsCustomization.cpp
@@ -91,8 +91,8 @@ void FNiagaraEmitterPropertiesDetails::BuildScriptProperties(TSharedRef<IPropert
 
 void FNiagaraEmitterPropertiesDetails::OnGenerateScalarConstantEntry(TSharedRef<IPropertyHandle> ElementProperty, int32 ElementIndex, IDetailChildrenBuilder& ChildrenBuilder)
 {
-	TSharedPtr<IPropertyHandle> ValueProperty = ElementProperty->GetChildHandle(GET_MEMBER_NAME_CHECKED(FNiagaraConstants_Float, Value));
-	TSharedPtr<IPropertyHandle> NameProperty = ElementProperty->GetChildHandle(GET_MEMBER_NAME_CHECKED(FNiagaraConstants_Float, Name));
+	TSharedPtr<IPropertyHandle> ValueProperty = ElementProperty->GetChildHandle(GET_MEMBER_NAME_CHECKED(FNiagaraConstant_Float, Value));
+	TSharedPtr<IPropertyHandle> NameProperty = ElementProperty->GetChildHandle(GET_MEMBER_NAME_CHECKED(FNiagaraConstant_Float, Name));
 
 	FName DisplayName;
 	NameProperty->GetValue(DisplayName);
@@ -105,8 +105,8 @@ void FNiagaraEmitterPropertiesDetails::OnGenerateScalarConstantEntry(TSharedRef<
 
 void FNiagaraEmitterPropertiesDetails::OnGenerateVectorConstantEntry(TSharedRef<IPropertyHandle> ElementProperty, int32 ElementIndex, IDetailChildrenBuilder& ChildrenBuilder)
 {
-	TSharedPtr<IPropertyHandle> ValueProperty = ElementProperty->GetChildHandle(GET_MEMBER_NAME_CHECKED(FNiagaraConstants_Vector, Value));
-	TSharedPtr<IPropertyHandle> NameProperty = ElementProperty->GetChildHandle(GET_MEMBER_NAME_CHECKED(FNiagaraConstants_Vector, Name));
+	TSharedPtr<IPropertyHandle> ValueProperty = ElementProperty->GetChildHandle(GET_MEMBER_NAME_CHECKED(FNiagaraConstant_Vector, Value));
+	TSharedPtr<IPropertyHandle> NameProperty = ElementProperty->GetChildHandle(GET_MEMBER_NAME_CHECKED(FNiagaraConstant_Vector, Name));
 
 	FName DisplayName;
 	NameProperty->GetValue(DisplayName);
@@ -119,8 +119,8 @@ void FNiagaraEmitterPropertiesDetails::OnGenerateVectorConstantEntry(TSharedRef<
 
 void FNiagaraEmitterPropertiesDetails::OnGenerateMatrixConstantEntry(TSharedRef<IPropertyHandle> ElementProperty, int32 ElementIndex, IDetailChildrenBuilder& ChildrenBuilder)
 {
-	TSharedPtr<IPropertyHandle> ValueProperty = ElementProperty->GetChildHandle(GET_MEMBER_NAME_CHECKED(FNiagaraConstants_Matrix, Value));
-	TSharedPtr<IPropertyHandle> NameProperty = ElementProperty->GetChildHandle(GET_MEMBER_NAME_CHECKED(FNiagaraConstants_Matrix, Name));
+	TSharedPtr<IPropertyHandle> ValueProperty = ElementProperty->GetChildHandle(GET_MEMBER_NAME_CHECKED(FNiagaraConstant_Matrix, Value));
+	TSharedPtr<IPropertyHandle> NameProperty = ElementProperty->GetChildHandle(GET_MEMBER_NAME_CHECKED(FNiagaraConstant_Matrix, Name));
 
 	FName DisplayName;
 	NameProperty->GetValue(DisplayName);
diff --git a/Engine/Source/Editor/NiagaraEditor/Private/NiagaraScriptSource.cpp b/Engine/Source/Editor/NiagaraEditor/Private/NiagaraScriptSource.cpp
--- a/Engine/Source/Editor/NiagaraEditor/Private/NiagaraScriptSource.cpp
+++ b/Engine/Source/Editor/NiagaraEditor/Private/NiagaraScriptSource.cpp
@@ -228,7 +228,7 @@ void UNiagaraScriptSource::GetEmitterAttributes(TArray<FName>& VectorInputs, TAr
 {
 	for (uint32 i=0; i < NiagaraConstants::NumBuiltinConstants; i++)
 	{
-		VectorInputs.Add(NiagaraConstants::ConstantNames[i]);
+		VectorInputs.Add(NiagaraConstants::GConstantNames[i]);
 	}
 
 	MatrixInputs.Empty();
diff --git a/Engine/Source/Editor/UnrealEd/Private/EditorBuildUtils.cpp b/Engine/Source/Editor/UnrealEd/Private/EditorBuildUtils.cpp
--- a/Engine/Source/Editor/UnrealEd/Private/EditorBuildUtils.cpp
+++ b/Engine/Source/Editor/UnrealEd/Private/EditorBuildUtils.cpp
@@ -1068,6 +1068,9 @@ bool FEditorBuildUtils::EditorBuildTextureStreaming(UWorld* InWorld, bool bWithT
 		FScopedSlowTask SlowTask((float)TexCoordScales.Num(), (LOCTEXT("TextureStreamingBuild_ExportingMaterialScales", "Exporting Material TexCoord Scales")));
 
 		const double StartTime = FPlatformTime::Seconds();
+
+		FMaterialUtilities::FExportErrorManager ExportErrors(FeatureLevel);
+
 		for (FTexCoordScaleMap::TIterator It(TexCoordScales); It; ++It)
 		{
 			SlowTask.EnterProgressFrame();
@@ -1078,9 +1081,11 @@ bool FEditorBuildUtils::EditorBuildTextureStreaming(UWorld* InWorld, bool bWithT
 			TArray<FMaterialTexCoordBuildInfo>& Scales = It.Value();
 
 
-			FMaterialUtilities::ExportMaterialTexCoordScales(MaterialInterface, QualityLevel, FeatureLevel, Scales);
+			FMaterialUtilities::ExportMaterialTexCoordScales(MaterialInterface, QualityLevel, FeatureLevel, Scales, ExportErrors);
 		}
 		UE_LOG(LogLevel, Display, TEXT("Export Material TexCoord Scales took %.3f seconds."), FPlatformTime::Seconds() - StartTime);
+
+		ExportErrors.OutputToLog();
 	}
 
 	if (bDebugDataOnly)
diff --git a/Engine/Source/Editor/UnrealEd/Private/Lightmass/LightmassRender.cpp b/Engine/Source/Editor/UnrealEd/Private/Lightmass/LightmassRender.cpp
--- a/Engine/Source/Editor/UnrealEd/Private/Lightmass/LightmassRender.cpp
+++ b/Engine/Source/Editor/UnrealEd/Private/Lightmass/LightmassRender.cpp
@@ -376,7 +376,7 @@ public:
 						return Compiler->Saturate(Compiler->ForceCast(MaterialInterface->CompileProperty(Compiler, DiffuseInput),MCT_Float3,true,true));
 					}
 				}
-				else if ((BlendMode == BLEND_Translucent) || (BlendMode == BLEND_Additive))
+				else if ((BlendMode == BLEND_Translucent) || (BlendMode == BLEND_Additive || (BlendMode == BLEND_AlphaComposite)))
 				{
 					int32 ColoredOpacity = -1;
 					if (ShadingModel == MSM_Unlit)
@@ -599,7 +599,8 @@ public:
 			else
 			if ((BlendMode == BLEND_Modulate) ||
 				(BlendMode == BLEND_Translucent) || 
-				(BlendMode == BLEND_Additive))
+				(BlendMode == BLEND_Additive) ||
+				(BlendMode == BLEND_AlphaComposite))
 			{
 				bool bColorInputIsNULL = false;
 				if (ShadingModel == MSM_Unlit)
@@ -611,7 +612,8 @@ public:
 					bColorInputIsNULL = !IsMaterialInputConnected(Material, MP_DiffuseColor);
 				}
 				if (BlendMode == BLEND_Translucent
-					|| BlendMode == BLEND_Additive)
+					|| BlendMode == BLEND_Additive
+					|| BlendMode == BLEND_AlphaComposite)
 				{
 					bExpressionIsNULL = bColorInputIsNULL && !IsMaterialInputConnected(Material, PropertyToCompile);
 				}
@@ -721,6 +723,7 @@ public:
 			break;
 		case BLEND_Translucent:
 		case BLEND_Additive:
+		case BLEND_AlphaComposite:
 			{
 				switch (InMaterialProperty)
 				{
diff --git a/Engine/Source/Editor/UnrealEd/Private/MaterialGraphNode.cpp b/Engine/Source/Editor/UnrealEd/Private/MaterialGraphNode.cpp
--- a/Engine/Source/Editor/UnrealEd/Private/MaterialGraphNode.cpp
+++ b/Engine/Source/Editor/UnrealEd/Private/MaterialGraphNode.cpp
@@ -457,15 +457,15 @@ void UMaterialGraphNode::CreateInputPins()
 		}
 		else if ( !FCString::Stricmp( *InputName, TEXT("AGreaterThanB") ) )
 		{
-			InputName = TEXT("A>=B");
+			InputName = TEXT("A > B");
 		}
 		else if ( !FCString::Stricmp( *InputName, TEXT("AEqualsB") ) )
 		{
-			InputName = TEXT("A==B");
+			InputName = TEXT("A == B");
 		}
 		else if ( !FCString::Stricmp( *InputName, TEXT("ALessThanB") ) )
 		{
-			InputName = TEXT("A<B");
+			InputName = TEXT("A < B");
 		}
 		else if ( !FCString::Stricmp( *InputName, TEXT("MipLevel") ) )
 		{
diff --git a/Engine/Source/Programs/UnrealBuildTool/Configuration/UEBuildPlatform.cs b/Engine/Source/Programs/UnrealBuildTool/Configuration/UEBuildPlatform.cs
--- a/Engine/Source/Programs/UnrealBuildTool/Configuration/UEBuildPlatform.cs
+++ b/Engine/Source/Programs/UnrealBuildTool/Configuration/UEBuildPlatform.cs
@@ -1009,7 +1009,11 @@ namespace UnrealBuildTool
 			{
 				String InstalledSDKVersionString = GetRequiredSDKString();
 				String PlatformSDKRoot = GetPathToPlatformAutoSDKs();
-				if (Directory.Exists(PlatformSDKRoot))
+                if (!Directory.Exists(PlatformSDKRoot))
+                {
+                    Directory.CreateDirectory(PlatformSDKRoot);
+                }
+
 				{
 					string VersionFilename = Path.Combine(PlatformSDKRoot, CurrentlyInstalledSDKStringManifest);
 					if (File.Exists(VersionFilename))
diff --git a/Engine/Source/Programs/UnrealLightmass/Private/ImportExport/Material.cpp b/Engine/Source/Programs/UnrealLightmass/Private/ImportExport/Material.cpp
--- a/Engine/Source/Programs/UnrealLightmass/Private/ImportExport/Material.cpp
+++ b/Engine/Source/Programs/UnrealLightmass/Private/ImportExport/Material.cpp
@@ -57,7 +57,7 @@ namespace Lightmass
 		else
 		{
 			// Materials with a translucent blend mode should always import transmission
-			check(BlendMode != BLEND_Translucent && BlendMode != BLEND_Additive && BlendMode != BLEND_Modulate);
+			check(BlendMode != BLEND_Translucent && BlendMode != BLEND_Additive && BlendMode != BLEND_Modulate && BlendMode != BLEND_AlphaComposite);
 		}
 
 		// Normal
diff --git a/Engine/Source/Programs/UnrealLightmass/Public/MaterialExport.h b/Engine/Source/Programs/UnrealLightmass/Public/MaterialExport.h
--- a/Engine/Source/Programs/UnrealLightmass/Public/MaterialExport.h
+++ b/Engine/Source/Programs/UnrealLightmass/Public/MaterialExport.h
@@ -45,7 +45,8 @@ enum EBlendMode
     BLEND_Translucent       =2,
     BLEND_Additive          =3,
     BLEND_Modulate          =4,
-    BLEND_MAX               =5,
+	BLEND_AlphaComposite    =5,
+    BLEND_MAX               =6,
 };
 
 struct FMaterialData
diff --git a/Engine/Source/Runtime/Apple/MetalRHI/Private/MetalRHI.cpp b/Engine/Source/Runtime/Apple/MetalRHI/Private/MetalRHI.cpp
--- a/Engine/Source/Runtime/Apple/MetalRHI/Private/MetalRHI.cpp
+++ b/Engine/Source/Runtime/Apple/MetalRHI/Private/MetalRHI.cpp
@@ -415,7 +415,8 @@ FMetalDynamicRHI::FMetalDynamicRHI()
 	GPixelFormats[PF_R8G8				].PlatformFormat	= MTLPixelFormatRG8Unorm;
 	GPixelFormats[PF_R16_SINT			].PlatformFormat	= MTLPixelFormatR16Sint;
 	GPixelFormats[PF_R16_UINT			].PlatformFormat	= MTLPixelFormatR16Uint;
-	
+	GPixelFormats[PF_R8_UINT			].PlatformFormat	= MTLPixelFormatR8Uint;
+
 	// get driver version (todo: share with other RHIs)
 	{
 		FGPUDriverInfo GPUDriverInfo = FPlatformMisc::GetGPUDriverInfo(GRHIAdapterName);
diff --git a/Engine/Source/Runtime/Core/Public/Stats/Stats2.h b/Engine/Source/Runtime/Core/Public/Stats/Stats2.h
--- a/Engine/Source/Runtime/Core/Public/Stats/Stats2.h
+++ b/Engine/Source/Runtime/Core/Public/Stats/Stats2.h
@@ -2029,6 +2029,7 @@ DECLARE_STATS_GROUP(TEXT("DDC"),STATGROUP_DDC, STATCAT_Advanced);
 DECLARE_STATS_GROUP(TEXT("Default Stat Group"),STATGROUP_Default, STATCAT_Advanced);
 DECLARE_STATS_GROUP(TEXT("Engine"),STATGROUP_Engine, STATCAT_Advanced);
 DECLARE_STATS_GROUP(TEXT("FPS Chart"),STATGROUP_FPSChart, STATCAT_Advanced);
+DECLARE_STATS_GROUP(TEXT("GPU"), STATGROUP_GPU, STATCAT_Advanced);
 DECLARE_STATS_GROUP(TEXT("GPU Particles"),STATGROUP_GPUParticles, STATCAT_Advanced);
 DECLARE_STATS_GROUP(TEXT("Game"),STATGROUP_Game, STATCAT_Advanced);
 DECLARE_STATS_GROUP(TEXT("GPU Defrag"), STATGROUP_GPUDEFRAG, STATCAT_Advanced);
diff --git a/Engine/Source/Runtime/CoreUObject/Classes/Object.h b/Engine/Source/Runtime/CoreUObject/Classes/Object.h
--- a/Engine/Source/Runtime/CoreUObject/Classes/Object.h
+++ b/Engine/Source/Runtime/CoreUObject/Classes/Object.h
@@ -162,6 +162,7 @@ enum EPixelFormat
 	PF_ASTC_12x12,
 	PF_BC6H,
 	PF_BC7,
+	PF_R8_UINT,
 	PF_MAX,
 };
 
diff --git a/Engine/Source/Runtime/Engine/Classes/Components/LightComponent.h b/Engine/Source/Runtime/Engine/Classes/Components/LightComponent.h
--- a/Engine/Source/Runtime/Engine/Classes/Components/LightComponent.h
+++ b/Engine/Source/Runtime/Engine/Classes/Components/LightComponent.h
@@ -59,6 +59,12 @@ class ENGINE_API ULightComponent : public ULightComponentBase
 	UPROPERTY(EditAnywhere, BlueprintReadOnly, interp, Category = Light, meta = (UIMin = "1700.0", UIMax = "12000.0"))
 	float Temperature;
 	
+	UPROPERTY(EditAnywhere, Category = Performance)
+	float MaxDrawDistance;
+
+	UPROPERTY(EditAnywhere, Category = Performance)
+	float MaxDistanceFadeRange;
+
 	/** false: use white (D65) as illuminant. */
 	UPROPERTY(EditAnywhere, BlueprintReadOnly, Category=Light, meta=(DisplayName = "Use Temperature"))
 	uint32 bUseTemperature : 1;
diff --git a/Engine/Source/Runtime/Engine/Classes/Engine/EngineTypes.h b/Engine/Source/Runtime/Engine/Classes/Engine/EngineTypes.h
--- a/Engine/Source/Runtime/Engine/Classes/Engine/EngineTypes.h
+++ b/Engine/Source/Runtime/Engine/Classes/Engine/EngineTypes.h
@@ -210,6 +210,7 @@ enum EBlendMode
 	BLEND_Translucent UMETA(DisplayName="Translucent"),
 	BLEND_Additive UMETA(DisplayName="Additive"),
 	BLEND_Modulate UMETA(DisplayName="Modulate"),
+	BLEND_AlphaComposite UMETA(DisplayName ="AlphaComposite (Premultiplied Alpha)"),
 	BLEND_MAX,
 };
 
diff --git a/Engine/Source/Runtime/Engine/Classes/Materials/MaterialExpressionIf.h b/Engine/Source/Runtime/Engine/Classes/Materials/MaterialExpressionIf.h
--- a/Engine/Source/Runtime/Engine/Classes/Materials/MaterialExpressionIf.h
+++ b/Engine/Source/Runtime/Engine/Classes/Materials/MaterialExpressionIf.h
@@ -19,7 +19,7 @@ class UMaterialExpressionIf : public UMaterialExpression
 	UPROPERTY()
 	FExpressionInput AGreaterThanB;
 
-	UPROPERTY(meta = (RequiredInput = "false"))
+	UPROPERTY(meta = (RequiredInput = "false", ToolTip = "Defaults to 'A > B' if not specified"))
 	FExpressionInput AEqualsB;
 
 	UPROPERTY()
diff --git a/Engine/Source/Runtime/Engine/Classes/Materials/MaterialExpressionNoise.h b/Engine/Source/Runtime/Engine/Classes/Materials/MaterialExpressionNoise.h
--- a/Engine/Source/Runtime/Engine/Classes/Materials/MaterialExpressionNoise.h
+++ b/Engine/Source/Runtime/Engine/Classes/Materials/MaterialExpressionNoise.h
@@ -29,8 +29,8 @@ enum ENoiseFunction
 	NOISEFUNCTION_GradientTex3D UMETA(DisplayName = "Fast Gradient - 3D Texture"),
 
 	/** High quality for direct use and bumps
-	 * Non-tiled: ~103 instructions per level, no textures
-	 * Tiling: ~174 instructions per level, no textures
+	 * Non-tiled: ~80 instructions per level, no textures
+	 * Tiling: ~143 instructions per level, no textures
 	 */
 	NOISEFUNCTION_GradientALU UMETA(DisplayName = "Gradient - Computational"),
 
@@ -41,6 +41,13 @@ enum ENoiseFunction
 	 */
 	NOISEFUNCTION_ValueALU UMETA(DisplayName="Value - Computational"),
 
+	/** Also known as Worley or Cellular noise
+	  * Quality=1 searches 8 cells, Quality=2 searches 16 cells
+	  * Quality=3 searches 27 cells, Quality=4 searches 32 cells
+	  * All are about 20 instructions per cell searched
+	  */
+	NOISEFUNCTION_VoronoiALU UMETA(DisplayName = "Voronoi"),
+
 	NOISEFUNCTION_MAX,
 };
 
@@ -61,8 +68,8 @@ class UMaterialExpressionNoise : public UMaterialExpression
 	UPROPERTY(EditAnywhere, Category=MaterialExpressionNoise)
 	float Scale;
 
-	/** 0 = fast, allows to pick a different implementation, can affect performance and look */
-	UPROPERTY(EditAnywhere, Category=MaterialExpressionNoise, meta=(UIMin = "1", UIMax = "10"))
+	/** Lower numbers are faster and lower quality, higher numbers are slower and higher quality */
+	UPROPERTY(EditAnywhere, Category=MaterialExpressionNoise, meta=(UIMin = "1", UIMax = "4"))
 	int32 Quality;
 
 	/** Noise function, affects performance and look */
diff --git a/Engine/Source/Runtime/Engine/Classes/Particles/Beam/ParticleModuleBeamNoise.h b/Engine/Source/Runtime/Engine/Classes/Particles/Beam/ParticleModuleBeamNoise.h
--- a/Engine/Source/Runtime/Engine/Classes/Particles/Beam/ParticleModuleBeamNoise.h
+++ b/Engine/Source/Runtime/Engine/Classes/Particles/Beam/ParticleModuleBeamNoise.h
@@ -16,6 +16,8 @@ class UParticleModuleBeamNoise : public UParticleModuleBeamBase
 {
 	GENERATED_UCLASS_BODY()
 
+	static const uint32 MaxNoiseTessellation = 500;
+
 	/** Is low frequency noise enabled. */
 	UPROPERTY(EditAnywhere, Category=LowFreq)
 	uint32 bLowFreq_Enabled:1;
diff --git a/Engine/Source/Runtime/Engine/Private/BatchedElements.cpp b/Engine/Source/Runtime/Engine/Private/BatchedElements.cpp
--- a/Engine/Source/Runtime/Engine/Private/BatchedElements.cpp
+++ b/Engine/Source/Runtime/Engine/Private/BatchedElements.cpp
@@ -102,6 +102,9 @@ void FBatchedElements::AddTriangle(int32 V0,int32 V1,int32 V2,const FTexture* Te
 	case BLEND_Modulate:
 		SimpleElementBlendMode = SE_BLEND_Modulate; 
 		break;
+	case BLEND_AlphaComposite:
+		SimpleElementBlendMode = SE_BLEND_AlphaComposite;
+		break;
 	};	
 	AddTriangle(V0,V1,V2,Texture,SimpleElementBlendMode);
 }
diff --git a/Engine/Source/Runtime/Engine/Private/Components/SceneCaptureComponent.cpp b/Engine/Source/Runtime/Engine/Private/Components/SceneCaptureComponent.cpp
--- a/Engine/Source/Runtime/Engine/Private/Components/SceneCaptureComponent.cpp
+++ b/Engine/Source/Runtime/Engine/Private/Components/SceneCaptureComponent.cpp
@@ -376,9 +376,7 @@ void USceneCaptureComponent2D::UpdateDeferredCaptures( FSceneInterface* Scene )
 	UWorld* World = Scene->GetWorld();
 	if( World && SceneCapturesToUpdateMap.Num() > 0 )
 	{
-		// We must push any deferred render state recreations before causing any rendering to happen, to make sure that deleted resource references are updated
-		World->SendAllEndOfFrameUpdates();
-		// Only update the scene captures assoicated with the current scene.
+		// Only update the scene captures associated with the current scene.
 		// Updating others not associated with the scene would cause invalid data to be rendered into the target
 		TArray< TWeakObjectPtr<USceneCaptureComponent2D> > SceneCapturesToUpdate;
 		SceneCapturesToUpdateMap.MultiFind( World, SceneCapturesToUpdate );
@@ -736,8 +734,6 @@ void USceneCaptureComponentCube::UpdateDeferredCaptures( FSceneInterface* Scene
 	
 	if( World && CubedSceneCapturesToUpdateMap.Num() > 0 )
 	{
-		// We must push any deferred render state recreations before causing any rendering to happen, to make sure that deleted resource references are updated
-		World->SendAllEndOfFrameUpdates();
 		// Only update the scene captures associated with the current scene.
 		// Updating others not associated with the scene would cause invalid data to be rendered into the target
 		TArray< TWeakObjectPtr<USceneCaptureComponentCube> > SceneCapturesToUpdate;
diff --git a/Engine/Source/Runtime/Engine/Private/DestructibleMesh.cpp b/Engine/Source/Runtime/Engine/Private/DestructibleMesh.cpp
--- a/Engine/Source/Runtime/Engine/Private/DestructibleMesh.cpp
+++ b/Engine/Source/Runtime/Engine/Private/DestructibleMesh.cpp
@@ -11,6 +11,7 @@
 #include "Engine/DestructibleFractureSettings.h"
 #include "PhysicsEngine/PhysXSupport.h"
 #include "EditorFramework/AssetImportData.h"
+#include "GPUSkinVertexFactory.h"
 
 #if WITH_APEX && WITH_EDITOR
 #include "ApexDestructibleAssetImport.h"
@@ -39,8 +40,8 @@ void UDestructibleMesh::PostLoad()
 	// you'll have to save re-chunked asset here, but I don't have a good way to let user
 	// know here since PostLoad invalidate any dirty package mark. 
 	FSkeletalMeshResource* ImportedMeshResource = GetImportedResource();
-	static IConsoleVariable* MaxBonesVar = IConsoleManager::Get().FindConsoleVariable(TEXT("Compat.MAX_GPUSKIN_BONES"));
-	const int32 MaxGPUSkinBones = MaxBonesVar->GetInt();
+	const uint32 MaxGPUSkinBones = FGPUBaseSkinVertexFactory::GetMaxGPUSkinBones();
+	check(MaxGPUSkinBones <= FGPUBaseSkinVertexFactory::GHardwareMaxGPUSkinBones);
 	// if this doesn't have the right MAX GPU Bone count, recreate it. 
 	for(int32 LodIndex=0; LodIndex<LODInfo.Num(); LodIndex++)
 	{
@@ -56,7 +57,7 @@ void UDestructibleMesh::PostLoad()
 
 		for (int32 SectionIndex = 0; SectionIndex < ThisLODModel.Sections.Num(); ++SectionIndex)
 		{
-			if (ThisLODModel.Sections[SectionIndex].BoneMap.Num() > MaxGPUSkinBones)
+			if (ThisLODModel.Sections[SectionIndex].BoneMap.Num() > (int)MaxGPUSkinBones)
 			{
 #if WITH_EDITOR
 				// re create destructible asset if it exceeds
diff --git a/Engine/Source/Runtime/Engine/Private/LevelTick.cpp b/Engine/Source/Runtime/Engine/Private/LevelTick.cpp
--- a/Engine/Source/Runtime/Engine/Private/LevelTick.cpp
+++ b/Engine/Source/Runtime/Engine/Private/LevelTick.cpp
@@ -1369,9 +1369,6 @@ void UWorld::Tick( ELevelTick TickType, float DeltaSeconds )
 	{
 		// Update SpeedTree wind objects.
 		Scene->UpdateSpeedTreeWind(TimeSeconds);
-
-		USceneCaptureComponent2D::UpdateDeferredCaptures( Scene );
-		USceneCaptureComponentCube::UpdateDeferredCaptures( Scene );
 	}
 
 	// Tick the FX system.
diff --git a/Engine/Source/Runtime/Engine/Private/Materials/MaterialExpressions.cpp b/Engine/Source/Runtime/Engine/Private/Materials/MaterialExpressions.cpp
--- a/Engine/Source/Runtime/Engine/Private/Materials/MaterialExpressions.cpp
+++ b/Engine/Source/Runtime/Engine/Private/Materials/MaterialExpressions.cpp
@@ -8807,7 +8807,7 @@ UMaterialExpressionNoise::UMaterialExpressionNoise(const FObjectInitializer& Obj
 
 	Scale = 1.0f;
 	Levels = 6;
-	Quality = 0;
+	Quality = 1;
 	OutputMin = -1.0f;
 	OutputMax = 1.0f;
 	LevelScale = 2.0f;
@@ -8829,7 +8829,9 @@ bool UMaterialExpressionNoise::CanEditChange(const UProperty* InProperty) const
 		FName PropertyFName = InProperty->GetFName();
 
 		bool bTilableNoiseType = NoiseFunction == NOISEFUNCTION_GradientALU || NoiseFunction == NOISEFUNCTION_ValueALU 
-			|| NoiseFunction == NOISEFUNCTION_GradientTex;
+			|| NoiseFunction == NOISEFUNCTION_GradientTex || NoiseFunction == NOISEFUNCTION_VoronoiALU;
+
+		bool bSupportsQuality = (NoiseFunction == NOISEFUNCTION_VoronoiALU);
 
 		if (PropertyFName == GET_MEMBER_NAME_CHECKED(UMaterialExpressionNoise, bTiling))
 		{
@@ -8839,6 +8841,11 @@ bool UMaterialExpressionNoise::CanEditChange(const UProperty* InProperty) const
 		{
 			bIsEditable = bTilableNoiseType && bTiling;
 		}
+
+		if (PropertyFName == GET_MEMBER_NAME_CHECKED(UMaterialExpressionNoise, Quality))
+		{
+			bIsEditable = bSupportsQuality;
+		}
 	}
 
 	return bIsEditable;
diff --git a/Engine/Source/Runtime/Engine/Private/Materials/MaterialShader.cpp b/Engine/Source/Runtime/Engine/Private/Materials/MaterialShader.cpp
--- a/Engine/Source/Runtime/Engine/Private/Materials/MaterialShader.cpp
+++ b/Engine/Source/Runtime/Engine/Private/Materials/MaterialShader.cpp
@@ -86,6 +86,7 @@ FString GetBlendModeString(EBlendMode BlendMode)
 		case BLEND_Translucent: BlendModeName = TEXT("BLEND_Translucent"); break;
 		case BLEND_Additive: BlendModeName = TEXT("BLEND_Additive"); break;
 		case BLEND_Modulate: BlendModeName = TEXT("BLEND_Modulate"); break;
+		case BLEND_AlphaComposite: BlendModeName = TEXT("BLEND_AlphaComposite"); break;
 		default: BlendModeName = TEXT("Unknown"); break;
 	}
 	return BlendModeName;
diff --git a/Engine/Source/Runtime/Engine/Private/Materials/MaterialShared.cpp b/Engine/Source/Runtime/Engine/Private/Materials/MaterialShared.cpp
--- a/Engine/Source/Runtime/Engine/Private/Materials/MaterialShared.cpp
+++ b/Engine/Source/Runtime/Engine/Private/Materials/MaterialShared.cpp
@@ -1377,6 +1377,11 @@ void FMaterial::SetupMaterialEnvironment(
 			}
 			break;
 		}
+	case BLEND_AlphaComposite: 
+		{
+			// Fall through the case, blend mode will reuse MATERIALBLENDING_TRANSLUCENT
+			OutEnvironment.SetDefine(TEXT("MATERIALBLENDING_ALPHACOMPOSITE"), TEXT("1"));
+		}
 	case BLEND_Translucent: OutEnvironment.SetDefine(TEXT("MATERIALBLENDING_TRANSLUCENT"),TEXT("1")); break;
 	case BLEND_Additive: OutEnvironment.SetDefine(TEXT("MATERIALBLENDING_ADDITIVE"),TEXT("1")); break;
 	case BLEND_Modulate: OutEnvironment.SetDefine(TEXT("MATERIALBLENDING_MODULATE"),TEXT("1")); break;
diff --git a/Engine/Source/Runtime/Engine/Private/Particles/ParticleBeamModules.cpp b/Engine/Source/Runtime/Engine/Private/Particles/ParticleBeamModules.cpp
--- a/Engine/Source/Runtime/Engine/Private/Particles/ParticleBeamModules.cpp
+++ b/Engine/Source/Runtime/Engine/Private/Particles/ParticleBeamModules.cpp
@@ -1575,6 +1575,13 @@ void UParticleModuleBeamNoise::PostEditChangeProperty(FPropertyChangedEvent& Pro
 	UProperty* PropertyThatChanged = PropertyChangedEvent.Property;
 	if (PartSys && PropertyThatChanged)
 	{
+		// Make sure that the interpolation count is > 0.
+		if (PropertyThatChanged->GetFName() == FName(TEXT("NoiseTessellation")))
+		{
+			// Clamp the tessellation
+			NoiseTessellation = FMath::Clamp<int32>(NoiseTessellation, 0, UParticleModuleBeamNoise::MaxNoiseTessellation);
+		}
+
 		PartSys->PostEditChangeProperty(PropertyChangedEvent);
 	}
 
diff --git a/Engine/Source/Runtime/Engine/Private/Particles/ParticleGpuSimulation.cpp b/Engine/Source/Runtime/Engine/Private/Particles/ParticleGpuSimulation.cpp
--- a/Engine/Source/Runtime/Engine/Private/Particles/ParticleGpuSimulation.cpp
+++ b/Engine/Source/Runtime/Engine/Private/Particles/ParticleGpuSimulation.cpp
@@ -33,6 +33,7 @@
 #include "GlobalDistanceFieldParameters.h"
 
 DECLARE_CYCLE_STAT(TEXT("GPUSpriteEmitterInstance Init"), STAT_GPUSpriteEmitterInstance_Init, STATGROUP_Particles);
+DECLARE_FLOAT_COUNTER_STAT(TEXT("Particle Simulation"), Stat_GPU_ParticleSimulation, STATGROUP_GPU);
 
 /*------------------------------------------------------------------------------
 	Constants to tune memory and performance for GPU particle simulation.
@@ -1390,6 +1391,8 @@ void ExecuteSimulationCommands(
 	}
 
 	SCOPED_DRAW_EVENT(RHICmdList, ParticleSimulation);
+	SCOPED_GPU_STAT(RHICmdList, Stat_GPU_ParticleSimulation);
+
 
 	const float FixDeltaSeconds = CVarGPUParticleFixDeltaSeconds.GetValueOnRenderThread();
 	const FParticleStateTextures& TextureResources = (FixDeltaSeconds <= 0 || bUseFixDT) ? ParticleSimulationResources->GetPreviousStateTextures() : ParticleSimulationResources->GetCurrentStateTextures();
@@ -1500,6 +1503,7 @@ void ClearTiles(FRHICommandList& RHICmdList, ERHIFeatureLevel::Type FeatureLevel
 	}
 
 	SCOPED_DRAW_EVENT(RHICmdList, ClearTiles);
+	SCOPED_GPU_STAT(RHICmdList, Stat_GPU_ParticleSimulation);
 
 	const int32 MaxTilesPerDrawCallUnaligned = GParticleScratchVertexBufferSize / sizeof(FVector2D);
 	const int32 MaxTilesPerDrawCall = MaxTilesPerDrawCallUnaligned & (~(TILES_PER_INSTANCE-1));
@@ -4657,6 +4661,7 @@ void FFXSystem::SimulateGPUParticles(
 	if (NewParticles.Num())
 	{
 		SCOPED_DRAW_EVENT(RHICmdList, ParticleInjection);
+		SCOPED_GPU_STAT(RHICmdList, Stat_GPU_ParticleSimulation);
 
 		FParticleStateTextures& CurrentStateTextures = ParticleSimulationResources->GetCurrentStateTextures();
 
diff --git a/Engine/Source/Runtime/Engine/Private/Particles/ParticleSystemRender.cpp b/Engine/Source/Runtime/Engine/Private/Particles/ParticleSystemRender.cpp
--- a/Engine/Source/Runtime/Engine/Private/Particles/ParticleSystemRender.cpp
+++ b/Engine/Source/Runtime/Engine/Private/Particles/ParticleSystemRender.cpp
@@ -500,7 +500,6 @@ void FDynamicSpriteEmitterDataBase::BuildViewFillData(
 
 	int32 NumIndices, IndexStride;
 	GetIndexAllocInfo(NumIndices, IndexStride);
-	check((uint32)NumIndices <= 65535);
 	check(IndexStride > 0);
 
 	DynamicIndexAllocation = FGlobalDynamicIndexBuffer::Get().Allocate( NumIndices, IndexStride );
@@ -1060,7 +1059,7 @@ void FDynamicSpriteEmitterData::GetDynamicMeshElementsEmitter(const FParticleSys
 					const FMaterial* Material = MaterialResource[bSelected]->GetMaterial(FeatureLevel);
 
 					if (Material && 
-						(Material->GetBlendMode() == BLEND_Translucent ||
+						(Material->GetBlendMode() == BLEND_Translucent || Material->GetBlendMode() == BLEND_AlphaComposite ||
 						((SourceData->SortMode == PSORTMODE_Age_OldestFirst) || (SourceData->SortMode == PSORTMODE_Age_NewestFirst)))
 						)
 					{
@@ -3194,6 +3193,11 @@ static int32 CreateDynamicBeam2EmitterIndices(TIndexType* OutIndex, const FDynam
 	TIndexType	VertexIndex = 0;
 	TIndexType	StartVertexIndex = 0;
 
+	TIndexType *BaseIndex = OutIndex;
+
+	// Signed as we are comparing against pointer arithmetic
+	const int32 MaxIndexCount = 65535;
+
 	for (int32 Beam = 0; Beam < Source.ActiveParticleCount; Beam++)
 	{
 		DECLARE_PARTICLE_PTR(Particle, Source.DataContainer.ParticleData + Source.ParticleStride * Beam);
@@ -3209,17 +3213,22 @@ static int32 CreateDynamicBeam2EmitterIndices(TIndexType* OutIndex, const FDynam
 
 		if (VertexIndex == 0)//First Beam
 		{
-			*(OutIndex++) = VertexIndex++;	// SheetIndex + 0
-			*(OutIndex++) = VertexIndex++;	// SheetIndex + 1
+			if ((OutIndex - BaseIndex <= MaxIndexCount - 2))
+			{
+				*(OutIndex++) = VertexIndex++;	// SheetIndex + 0
+				*(OutIndex++) = VertexIndex++;	// SheetIndex + 1
+			}
 		}
 		else//Degenerate tris between beams
 		{
-			*(OutIndex++) = VertexIndex - 1;	// Last vertex of the previous sheet
-			*(OutIndex++) = VertexIndex;		// First vertex of the next sheet
-			*(OutIndex++) = VertexIndex++;		// First vertex of the next sheet
-			*(OutIndex++) = VertexIndex++;		// Second vertex of the next sheet
-
-			TrianglesToRender += 4;
+			if ((OutIndex - BaseIndex <= MaxIndexCount - 4))
+			{
+				*(OutIndex++) = VertexIndex - 1;	// Last vertex of the previous sheet
+				*(OutIndex++) = VertexIndex;		// First vertex of the next sheet
+				*(OutIndex++) = VertexIndex++;		// First vertex of the next sheet
+				*(OutIndex++) = VertexIndex++;		// Second vertex of the next sheet
+				TrianglesToRender += 4;
+			}
 		}
 
 		for (int32 SheetIndex = 0; SheetIndex < Source.Sheets; SheetIndex++)
@@ -3231,10 +3240,16 @@ static int32 CreateDynamicBeam2EmitterIndices(TIndexType* OutIndex, const FDynam
 			for (int32 i = 0; i < BeamPayloadData->TriangleCount; i++)
 			{
 				*(OutIndex++) = VertexIndex++;
+
+				if (OutIndex - BaseIndex > MaxIndexCount)
+				{
+					break;
+				}
 			}
 
 			// Degenerate tris
-			if ((SheetIndex + 1) < Source.Sheets)
+			if ((SheetIndex + 1) < Source.Sheets
+				&& (OutIndex - BaseIndex <= MaxIndexCount-4))
 			{
 				*(OutIndex++) = VertexIndex - 1;	// Last vertex of the previous sheet
 				*(OutIndex++) = VertexIndex;		// First vertex of the next sheet
@@ -3243,6 +3258,12 @@ static int32 CreateDynamicBeam2EmitterIndices(TIndexType* OutIndex, const FDynam
 
 				TrianglesToRender += 4;
 			}
+
+			if (OutIndex - BaseIndex > MaxIndexCount)
+			{
+				break;
+			}
+
 		}
 	}
 
diff --git a/Engine/Source/Runtime/Engine/Private/Particles/ParticleTrail2EmitterInstance.cpp b/Engine/Source/Runtime/Engine/Private/Particles/ParticleTrail2EmitterInstance.cpp
--- a/Engine/Source/Runtime/Engine/Private/Particles/ParticleTrail2EmitterInstance.cpp
+++ b/Engine/Source/Runtime/Engine/Private/Particles/ParticleTrail2EmitterInstance.cpp
@@ -2185,7 +2185,7 @@ bool FParticleRibbonEmitterInstance::ResolveSourcePoint(int32 InTrailIdx,
 					ResolveSource();
 				}
 
-				if (SourceEmitter && SourceEmitter->ParticleIndices)
+				if (SourceEmitter && SourceEmitter->ParticleIndices && SourceEmitter->ActiveParticles > 0)
 				{
 					if (SourceIndices[InTrailIdx] != -1)
 					{
@@ -2227,6 +2227,7 @@ bool FParticleRibbonEmitterInstance::ResolveSourcePoint(int32 InTrailIdx,
 								if (SourceIndices[TrailCheckIdx] == SourceEmitter->ParticleIndices[Index])
 								{
 									Index = -1;
+									break;
 								}
 							}
 						}
diff --git a/Engine/Source/Runtime/Engine/Private/PointLightSceneProxy.h b/Engine/Source/Runtime/Engine/Private/PointLightSceneProxy.h
--- a/Engine/Source/Runtime/Engine/Private/PointLightSceneProxy.h
+++ b/Engine/Source/Runtime/Engine/Private/PointLightSceneProxy.h
@@ -39,6 +39,8 @@ public:
 	,	SourceRadius(Component->SourceRadius)
 	,	SourceLength(Component->SourceLength)
 	,	bInverseSquared(Component->bUseInverseSquaredFalloff)
+	,	MaxDrawDistance(Component->MaxDrawDistance)
+	,	FadeRange(Component->MaxDistanceFadeRange)
 	{
 		UpdateRadius(Component->AttenuationRadius);
 	}
@@ -50,6 +52,15 @@ public:
 	void UpdateRadius_GameThread(UPointLightComponent* Component);
 
 	// FLightSceneInfo interface.
+	virtual float GetMaxDrawDistance() const final override 
+	{ 
+		return MaxDrawDistance; 
+	}
+
+	virtual float GetFadeRange() const final override 
+	{ 
+		return FadeRange; 
+	}
 
 	/** @return radius of the light or 0 if no radius */
 	virtual float GetRadius() const override
@@ -147,6 +158,9 @@ private:
 		// Min to avoid div by 0 (NaN in InvRadius)
 		InvRadius = 1.0f / FMath::Max(0.00001f, ComponentRadius);
 	}
+
+	float MaxDrawDistance;
+	float FadeRange;
 };
 
 /**
