Fix UE4Game XboxOne compile errors, and UE4Editor Win64 static analysis warnings in non-unity.

diff --git a/Engine/Plugins/Editor/AudioCapture/Source/AudioCapture/Private/RtAudio.cpp b/Engine/Plugins/Editor/AudioCapture/Source/AudioCapture/Private/RtAudio.cpp
--- a/Engine/Plugins/Editor/AudioCapture/Source/AudioCapture/Private/RtAudio.cpp
+++ b/Engine/Plugins/Editor/AudioCapture/Source/AudioCapture/Private/RtAudio.cpp
@@ -44,11 +44,13 @@ WITH THE SOFTWARE OR THE USE OR OTHER DEALINGS IN THE SOFTWARE.
 
 #include "AllowWindowsPlatformTypes.h"
 
-#include <iostream>
-#include <cstdlib>
-#include <cstring>
-#include <climits>
-#include <algorithm>
+THIRD_PARTY_INCLUDES_START
+	#include <iostream>
+	#include <cstdlib>
+	#include <cstring>
+	#include <climits>
+	#include <algorithm>
+THIRD_PARTY_INCLUDES_END
 
 // Static variable definitions.
 const unsigned int RtApi::MAX_SAMPLE_RATES = 14;
diff --git a/Engine/Plugins/Editor/AudioCapture/Source/AudioCapture/Private/RtAudio.h b/Engine/Plugins/Editor/AudioCapture/Source/AudioCapture/Private/RtAudio.h
--- a/Engine/Plugins/Editor/AudioCapture/Source/AudioCapture/Private/RtAudio.h
+++ b/Engine/Plugins/Editor/AudioCapture/Source/AudioCapture/Private/RtAudio.h
@@ -49,10 +49,12 @@ WITH THE SOFTWARE OR THE USE OR OTHER DEALINGS IN THE SOFTWARE.
 
 #define RTAUDIO_VERSION "4.1.2"
 
-#include <string>
-#include <vector>
-#include <exception>
-#include <iostream>
+THIRD_PARTY_INCLUDES_START
+	#include <string>
+	#include <vector>
+	#include <exception>
+	#include <iostream>
+THIRD_PARTY_INCLUDES_END
 
 /*! \typedef typedef unsigned long RtAudioFormat;
 \brief RtAudio data format type.
diff --git a/Engine/Source/Developer/VulkanShaderFormat/Private/VulkanGlslangBridge.cpp b/Engine/Source/Developer/VulkanShaderFormat/Private/VulkanGlslangBridge.cpp
--- a/Engine/Source/Developer/VulkanShaderFormat/Private/VulkanGlslangBridge.cpp
+++ b/Engine/Source/Developer/VulkanShaderFormat/Private/VulkanGlslangBridge.cpp
@@ -14,13 +14,15 @@
 	#pragma warning(disable: 6340) // warning C6340: Mismatch on sign: 'const unsigned int' passed as _Param_(4) when some signed type is required in call to 'snprintf'.
 #endif
 
-#include "glslang/Include/ShHandle.h"
-#include "glslang/Include/revision.h"
-#include "glslang/Public/ShaderLang.h"
-#include "SPIRV/GlslangToSpv.h"
-#include "SPIRV/GLSL.std.450.h"
-#include "SPIRV/doc.h"
-#include "SPIRV/disassemble.h"
+THIRD_PARTY_INCLUDES_START
+	#include "glslang/Include/ShHandle.h"
+	#include "glslang/Include/revision.h"
+	#include "glslang/Public/ShaderLang.h"
+	#include "SPIRV/GlslangToSpv.h"
+	#include "SPIRV/GLSL.std.450.h"
+	#include "SPIRV/doc.h"
+	#include "SPIRV/disassemble.h"
+THIRD_PARTY_INCLUDES_END
 
 #if _MSC_VER >= 1800
 	#pragma warning(pop)
diff --git a/Engine/Source/Editor/LocalizationDashboard/Private/LocalizationTargetDetailCustomization.h b/Engine/Source/Editor/LocalizationDashboard/Private/LocalizationTargetDetailCustomization.h
--- a/Engine/Source/Editor/LocalizationDashboard/Private/LocalizationTargetDetailCustomization.h
+++ b/Engine/Source/Editor/LocalizationDashboard/Private/LocalizationTargetDetailCustomization.h
@@ -7,6 +7,7 @@
 #include "Widgets/Views/STableViewBase.h"
 #include "Widgets/Views/STableRow.h"
 #include "IDetailCustomization.h"
+#include "LocalizationTargetTypes.h"
 
 class IDetailLayoutBuilder;
 class IPropertyHandle;
diff --git a/Engine/Source/Editor/SequenceRecorder/Private/SSequenceRecorder.h b/Engine/Source/Editor/SequenceRecorder/Private/SSequenceRecorder.h
--- a/Engine/Source/Editor/SequenceRecorder/Private/SSequenceRecorder.h
+++ b/Engine/Source/Editor/SequenceRecorder/Private/SSequenceRecorder.h
@@ -8,6 +8,7 @@
 #include "Widgets/SCompoundWidget.h"
 #include "Widgets/Views/STableViewBase.h"
 #include "Widgets/Views/STableRow.h"
+#include "ActorRecording.h"
 
 class FActiveTimerHandle;
 class FUICommandList;
diff --git a/Engine/Source/Programs/UnrealCEFSubProcess/Private/UnrealCEFSubProcess.h b/Engine/Source/Programs/UnrealCEFSubProcess/Private/UnrealCEFSubProcess.h
--- a/Engine/Source/Programs/UnrealCEFSubProcess/Private/UnrealCEFSubProcess.h
+++ b/Engine/Source/Programs/UnrealCEFSubProcess/Private/UnrealCEFSubProcess.h
@@ -12,7 +12,9 @@
 #endif
 #pragma push_macro("OVERRIDE")
 #undef OVERRIDE // cef headers provide their own OVERRIDE macro
+THIRD_PARTY_INCLUDES_START
 #include "include/cef_app.h"
+THIRD_PARTY_INCLUDES_END
 #pragma pop_macro("OVERRIDE")
 #if PLATFORM_WINDOWS
 #include "HideWindowsPlatformAtomics.h"
diff --git a/Engine/Source/Programs/UnrealCEFSubProcess/Private/UnrealCEFSubProcessApp.h b/Engine/Source/Programs/UnrealCEFSubProcess/Private/UnrealCEFSubProcessApp.h
--- a/Engine/Source/Programs/UnrealCEFSubProcess/Private/UnrealCEFSubProcessApp.h
+++ b/Engine/Source/Programs/UnrealCEFSubProcess/Private/UnrealCEFSubProcessApp.h
@@ -13,7 +13,9 @@
 
 #pragma push_macro("OVERRIDE")
 #undef OVERRIDE // cef headers provide their own OVERRIDE macro
+THIRD_PARTY_INCLUDES_START
 #include "include/cef_app.h"
+THIRD_PARTY_INCLUDES_END
 #pragma pop_macro("OVERRIDE")
 
 #if PLATFORM_WINDOWS
diff --git a/Engine/Source/Programs/UnrealCEFSubProcess/Private/UnrealCEFSubProcessRemoteMethodHandler.h b/Engine/Source/Programs/UnrealCEFSubProcess/Private/UnrealCEFSubProcessRemoteMethodHandler.h
--- a/Engine/Source/Programs/UnrealCEFSubProcess/Private/UnrealCEFSubProcessRemoteMethodHandler.h
+++ b/Engine/Source/Programs/UnrealCEFSubProcess/Private/UnrealCEFSubProcessRemoteMethodHandler.h
@@ -13,7 +13,9 @@
 
 #pragma push_macro("OVERRIDE")
 #undef OVERRIDE // cef headers provide their own OVERRIDE macro
+THIRD_PARTY_INCLUDES_START
 #include "include/cef_v8.h"
+THIRD_PARTY_INCLUDES_END
 #pragma pop_macro("OVERRIDE")
 
 #if PLATFORM_WINDOWS
diff --git a/Engine/Source/Programs/UnrealCEFSubProcess/Private/UnrealCEFSubProcessRemoteScripting.h b/Engine/Source/Programs/UnrealCEFSubProcess/Private/UnrealCEFSubProcessRemoteScripting.h
--- a/Engine/Source/Programs/UnrealCEFSubProcess/Private/UnrealCEFSubProcessRemoteScripting.h
+++ b/Engine/Source/Programs/UnrealCEFSubProcess/Private/UnrealCEFSubProcessRemoteScripting.h
@@ -12,8 +12,10 @@
 
 #pragma push_macro("OVERRIDE")
 #undef OVERRIDE // cef headers provide their own OVERRIDE macro
+THIRD_PARTY_INCLUDES_START
 #include "include/cef_app.h"
 #include "include/cef_v8.h"
+THIRD_PARTY_INCLUDES_END
 #pragma pop_macro("OVERRIDE")
 
 #if PLATFORM_WINDOWS
diff --git a/Engine/Source/Runtime/Engine/Private/CaptureSource.h b/Engine/Source/Runtime/Engine/Private/CaptureSource.h
--- a/Engine/Source/Runtime/Engine/Private/CaptureSource.h
+++ b/Engine/Source/Runtime/Engine/Private/CaptureSource.h
@@ -21,6 +21,7 @@ typedef TCHAR* PTCHAR;
 #if USING_CODE_ANALYSIS
 	#pragma warning(disable:6509) // Invalid annotation: 'return' cannot be referenced in some contexts
 	#pragma warning(disable:6101) // Returning uninitialized memory '*lpdwExitCode'.  A successful path through the function does not set the named _Out_ parameter.
+	#pragma warning(disable:28204) // 'Func' has an override at 'file' and only the override is annotated for _Param_(N): when an override is annotated, the base (this function) should be similarly annotated.
 #endif
 #include "WindowsHWrapper.h"
 #include "AllowWindowsPlatformTypes.h"
diff --git a/Engine/Source/Runtime/WebBrowser/Private/CEF/CEFBrowserByteResource.h b/Engine/Source/Runtime/WebBrowser/Private/CEF/CEFBrowserByteResource.h
--- a/Engine/Source/Runtime/WebBrowser/Private/CEF/CEFBrowserByteResource.h
+++ b/Engine/Source/Runtime/WebBrowser/Private/CEF/CEFBrowserByteResource.h
@@ -14,7 +14,9 @@
 
 #pragma push_macro("OVERRIDE")
 #undef OVERRIDE // cef headers provide their own OVERRIDE macro
+THIRD_PARTY_INCLUDES_START
 #include "include/cef_resource_handler.h"
+THIRD_PARTY_INCLUDES_END
 #pragma pop_macro("OVERRIDE")
 
 #if PLATFORM_WINDOWS
diff --git a/Engine/Source/Runtime/WebBrowser/Private/CEF/CEFBrowserHandler.h b/Engine/Source/Runtime/WebBrowser/Private/CEF/CEFBrowserHandler.h
--- a/Engine/Source/Runtime/WebBrowser/Private/CEF/CEFBrowserHandler.h
+++ b/Engine/Source/Runtime/WebBrowser/Private/CEF/CEFBrowserHandler.h
@@ -19,7 +19,9 @@ struct Rect;
 
 #pragma push_macro("OVERRIDE")
 #undef OVERRIDE // cef headers provide their own OVERRIDE macro
+THIRD_PARTY_INCLUDES_START
 #include "include/cef_client.h"
+THIRD_PARTY_INCLUDES_END
 #pragma pop_macro("OVERRIDE")
 
 #if PLATFORM_WINDOWS
diff --git a/Engine/Source/Runtime/WebBrowser/Private/CEF/CEFBrowserPopupFeatures.h b/Engine/Source/Runtime/WebBrowser/Private/CEF/CEFBrowserPopupFeatures.h
--- a/Engine/Source/Runtime/WebBrowser/Private/CEF/CEFBrowserPopupFeatures.h
+++ b/Engine/Source/Runtime/WebBrowser/Private/CEF/CEFBrowserPopupFeatures.h
@@ -18,7 +18,9 @@ class IWebBrowserPopupFeatures;
 
 #pragma push_macro("OVERRIDE")
 #undef OVERRIDE // cef headers provide their own OVERRIDE macro
+THIRD_PARTY_INCLUDES_START
 #include "include/cef_base.h"
+THIRD_PARTY_INCLUDES_END
 #pragma pop_macro("OVERRIDE")
 
 #if PLATFORM_WINDOWS
diff --git a/Engine/Source/Runtime/WebBrowser/Private/CEF/CEFCookieManager.cpp b/Engine/Source/Runtime/WebBrowser/Private/CEF/CEFCookieManager.cpp
--- a/Engine/Source/Runtime/WebBrowser/Private/CEF/CEFCookieManager.cpp
+++ b/Engine/Source/Runtime/WebBrowser/Private/CEF/CEFCookieManager.cpp
@@ -13,7 +13,9 @@
 #endif
 #pragma push_macro("OVERRIDE")
 #	undef OVERRIDE // cef headers provide their own OVERRIDE macro
+THIRD_PARTY_INCLUDES_START
 #	include "include/cef_app.h"
+THIRD_PARTY_INCLUDES_END
 #pragma pop_macro("OVERRIDE")
 #if PLATFORM_WINDOWS
 #	include "HideWindowsPlatformAtomics.h"
diff --git a/Engine/Source/Runtime/WebBrowser/Private/CEF/CEFJSScripting.h b/Engine/Source/Runtime/WebBrowser/Private/CEF/CEFJSScripting.h
--- a/Engine/Source/Runtime/WebBrowser/Private/CEF/CEFJSScripting.h
+++ b/Engine/Source/Runtime/WebBrowser/Private/CEF/CEFJSScripting.h
@@ -18,8 +18,10 @@ struct FWebJSParam;
 #endif
 #pragma push_macro("OVERRIDE")
 #undef OVERRIDE // cef headers provide their own OVERRIDE macro
+THIRD_PARTY_INCLUDES_START
 #include "include/cef_client.h"
 #include "include/cef_values.h"
+THIRD_PARTY_INCLUDES_END
 #pragma pop_macro("OVERRIDE")
 #if PLATFORM_WINDOWS
 #include "HideWindowsPlatformAtomics.h"
diff --git a/Engine/Source/Runtime/WebBrowser/Private/CEF/CEFJSStructDeserializerBackend.h b/Engine/Source/Runtime/WebBrowser/Private/CEF/CEFJSStructDeserializerBackend.h
--- a/Engine/Source/Runtime/WebBrowser/Private/CEF/CEFJSStructDeserializerBackend.h
+++ b/Engine/Source/Runtime/WebBrowser/Private/CEF/CEFJSStructDeserializerBackend.h
@@ -25,7 +25,9 @@ class UStruct;
 #endif
 #pragma push_macro("OVERRIDE")
 #undef OVERRIDE // cef headers provide their own OVERRIDE macro
+THIRD_PARTY_INCLUDES_START
 #include "include/cef_values.h"
+THIRD_PARTY_INCLUDES_END
 #pragma pop_macro("OVERRIDE")
 #if PLATFORM_WINDOWS
 #include "HideWindowsPlatformAtomics.h"
diff --git a/Engine/Source/Runtime/WebBrowser/Private/CEF/CEFJSStructSerializerBackend.h b/Engine/Source/Runtime/WebBrowser/Private/CEF/CEFJSStructSerializerBackend.h
--- a/Engine/Source/Runtime/WebBrowser/Private/CEF/CEFJSStructSerializerBackend.h
+++ b/Engine/Source/Runtime/WebBrowser/Private/CEF/CEFJSStructSerializerBackend.h
@@ -31,7 +31,9 @@ class UStruct;
 
 #pragma push_macro("OVERRIDE")
 #undef OVERRIDE // cef headers provide their own OVERRIDE macro
+THIRD_PARTY_INCLUDES_START
 #include "include/cef_values.h"
+THIRD_PARTY_INCLUDES_END
 #pragma pop_macro("OVERRIDE")
 
 #if PLATFORM_WINDOWS
diff --git a/Engine/Source/Runtime/WebBrowser/Private/CEF/CEFWebBrowserWindow.h b/Engine/Source/Runtime/WebBrowser/Private/CEF/CEFWebBrowserWindow.h
--- a/Engine/Source/Runtime/WebBrowser/Private/CEF/CEFWebBrowserWindow.h
+++ b/Engine/Source/Runtime/WebBrowser/Private/CEF/CEFWebBrowserWindow.h
@@ -29,12 +29,14 @@ enum class EWebBrowserDocumentState;
 	#include "AllowWindowsPlatformAtomics.h"
 #endif
 
+THIRD_PARTY_INCLUDES_START
 #pragma push_macro("OVERRIDE")
 #	undef OVERRIDE // cef headers provide their own OVERRIDE macro
 #	include "include/internal/cef_ptr.h"
 #	include "include/cef_render_handler.h"
 #	include "include/cef_jsdialog_handler.h"
 #pragma pop_macro("OVERRIDE")
+THIRD_PARTY_INCLUDES_END
 
 #if PLATFORM_WINDOWS
 	#include "HideWindowsPlatformAtomics.h"
diff --git a/Engine/Source/Runtime/WebBrowser/Private/WebBrowserSingleton.cpp b/Engine/Source/Runtime/WebBrowser/Private/WebBrowserSingleton.cpp
--- a/Engine/Source/Runtime/WebBrowser/Private/WebBrowserSingleton.cpp
+++ b/Engine/Source/Runtime/WebBrowser/Private/WebBrowserSingleton.cpp
@@ -25,7 +25,9 @@
 #	endif
 #	pragma push_macro("OVERRIDE")
 #		undef OVERRIDE // cef headers provide their own OVERRIDE macro
+THIRD_PARTY_INCLUDES_START
 #		include "include/cef_app.h"
+THIRD_PARTY_INCLUDES_END
 #	pragma pop_macro("OVERRIDE")
 #	if PLATFORM_WINDOWS
 #		include "HideWindowsPlatformTypes.h"
diff --git a/Engine/Source/Runtime/WebBrowser/Private/WebBrowserSingleton.h b/Engine/Source/Runtime/WebBrowser/Private/WebBrowserSingleton.h
--- a/Engine/Source/Runtime/WebBrowser/Private/WebBrowserSingleton.h
+++ b/Engine/Source/Runtime/WebBrowser/Private/WebBrowserSingleton.h
@@ -20,8 +20,10 @@ struct FWebBrowserWindowInfo;
 #endif
 #pragma push_macro("OVERRIDE")
 #undef OVERRIDE // cef headers provide their own OVERRIDE macro
+THIRD_PARTY_INCLUDES_START
 #include "include/internal/cef_ptr.h"
 #include "include/cef_request_context.h"
+THIRD_PARTY_INCLUDES_END
 #pragma pop_macro("OVERRIDE")
 #if PLATFORM_WINDOWS
 	#include "HideWindowsPlatformAtomics.h"
diff --git a/Engine/Source/Runtime/WebBrowser/Private/WebJSFunction.cpp b/Engine/Source/Runtime/WebBrowser/Private/WebJSFunction.cpp
--- a/Engine/Source/Runtime/WebBrowser/Private/WebJSFunction.cpp
+++ b/Engine/Source/Runtime/WebBrowser/Private/WebJSFunction.cpp
@@ -11,7 +11,9 @@
 #endif
 #pragma push_macro("OVERRIDE")
 #undef OVERRIDE // cef headers provide their own OVERRIDE macro
+THIRD_PARTY_INCLUDES_START
 #include "include/cef_values.h"
+THIRD_PARTY_INCLUDES_END
 #pragma pop_macro("OVERRIDE")
 #if PLATFORM_WINDOWS
 #include "HideWindowsPlatformAtomics.h"
