UE4.20 - VS 2015 Compile fix

diff --git a/Engine/Source/Developer/VulkanShaderFormat/Private/VulkanGlslangBridge.cpp b/Engine/Source/Developer/VulkanShaderFormat/Private/VulkanGlslangBridge.cpp
--- a/Engine/Source/Developer/VulkanShaderFormat/Private/VulkanGlslangBridge.cpp
+++ b/Engine/Source/Developer/VulkanShaderFormat/Private/VulkanGlslangBridge.cpp
@@ -408,7 +408,7 @@ bool GenerateSpirv(const ANSICHAR* Source, FCompilerInfo& CompilerInfo, FString&
 			{
 				Binding = Program->getUniformBinding(Index);
 			}
-			FSpirv::FEntry Entry{Name, Binding};
+			FSpirv::FEntry Entry(Name, Binding);
 			OutSpirv.ReflectionInfo.Add(Entry);
 		}
 
@@ -426,7 +426,7 @@ bool GenerateSpirv(const ANSICHAR* Source, FCompilerInfo& CompilerInfo, FString&
 			{
 				Binding = Program->getUniformBinding(Index);
 			}
-			FSpirv::FEntry Entry{Name, Binding};
+			FSpirv::FEntry Entry(Name, Binding);
 			OutSpirv.ReflectionInfo.Add(Entry);
 		}
 
diff --git a/Engine/Source/Developer/VulkanShaderFormat/Private/VulkanShaderFormat.h b/Engine/Source/Developer/VulkanShaderFormat/Private/VulkanShaderFormat.h
--- a/Engine/Source/Developer/VulkanShaderFormat/Private/VulkanShaderFormat.h
+++ b/Engine/Source/Developer/VulkanShaderFormat/Private/VulkanShaderFormat.h
@@ -39,6 +39,17 @@ struct FSpirv
 	TArray<uint32> Data;
 	struct FEntry
 	{
+		FEntry() = default;
+
+		FEntry(const FString& InName, int32 InBinding)
+			: Name(InName)
+			, Binding(-InBinding)
+		{
+		}
+
+		FString Name;
+		int32 Binding = -1;
+
 		FString Name;
 		int32 Binding = -1;
 
