UE4.20 - Static analysis fix

diff --git a/Engine/Source/Developer/VulkanShaderFormat/Private/VulkanGlslangBridge.cpp b/Engine/Source/Developer/VulkanShaderFormat/Private/VulkanGlslangBridge.cpp
--- a/Engine/Source/Developer/VulkanShaderFormat/Private/VulkanGlslangBridge.cpp
+++ b/Engine/Source/Developer/VulkanShaderFormat/Private/VulkanGlslangBridge.cpp
@@ -248,7 +248,7 @@ static void ComputeMovableWordIndices(FSpirv& Spirv)
 		{
 			uint32 Result = Ptr[1];
 			spv::StorageClass Storage = (spv::StorageClass)Ptr[2];
-			if (Storage == spv::StorageClassUniform || spv::StorageClassUniformConstant)
+			if (Storage == spv::StorageClassUniform || Storage == spv::StorageClassUniformConstant)
 			{
 				uint32 Type = Ptr[3];
 				TypePointerUniforms.Add(Result, Type);
