UE4.20 - Merge fix

diff --git a/Engine/Source/Developer/VulkanShaderFormat/Private/VulkanShaderFormat.h b/Engine/Source/Developer/VulkanShaderFormat/Private/VulkanShaderFormat.h
--- a/Engine/Source/Developer/VulkanShaderFormat/Private/VulkanShaderFormat.h
+++ b/Engine/Source/Developer/VulkanShaderFormat/Private/VulkanShaderFormat.h
@@ -50,9 +50,6 @@ struct FSpirv
 		FString Name;
 		int32 Binding = -1;
 
-		FString Name;
-		int32 Binding = -1;
-
 		uint32 DescriptorSet = UINT32_MAX;
 
 		// Index into the Spirv Word containing the descriptor set decoration
