UE4.22 - Fix Lumin shaders format
#rnx

diff --git a/Engine/Source/Developer/Lumin/LuminPlatformEditor/Private/LuminPlatformEditorModule.cpp b/Engine/Source/Developer/Lumin/LuminPlatformEditor/Private/LuminPlatformEditorModule.cpp
--- a/Engine/Source/Developer/Lumin/LuminPlatformEditor/Private/LuminPlatformEditorModule.cpp
+++ b/Engine/Source/Developer/Lumin/LuminPlatformEditor/Private/LuminPlatformEditorModule.cpp
@@ -71,7 +71,10 @@ class FLuminPlatformEditorModule
 
 			{
 				static FName NAME_SF_VULKAN_ES31_LUMIN(TEXT("SF_VULKAN_ES31_LUMIN"));
-				UShaderPlatformQualitySettings* LuminMaterialQualitySettings = UMaterialShaderQualitySettings::Get()->GetShaderPlatformQualitySettings(NAME_SF_VULKAN_ES31_LUMIN);
+				static FName NAME_SF_VULKAN_ES31_LUMIN_NOUB(TEXT("SF_VULKAN_ES31_LUMIN_NOUB"));
+				auto* CVar = IConsoleManager::Get().FindTConsoleVariableDataInt(TEXT("r.Vulkan.UseRealUBs"));
+				const bool bUseNOUB = (CVar && CVar->GetValueOnAnyThread() == 0);
+				UShaderPlatformQualitySettings* LuminMaterialQualitySettings = UMaterialShaderQualitySettings::Get()->GetShaderPlatformQualitySettings(bUseNOUB ? NAME_SF_VULKAN_ES31_LUMIN_NOUB : NAME_SF_VULKAN_ES31_LUMIN);
 				SettingsModule->RegisterSettings("Project", "Platforms", "MagicLeapVulkanQuality",
 					LOCTEXT("LuminVulkanQualitySettingsName", "Lumin Material Quality - Vulkan"),
 					LOCTEXT("LuminVulkanQualitySettingsDescription", "Settings for Lumin Vulkan material quality."),
diff --git a/Engine/Source/Developer/Lumin/LuminTargetPlatform/Private/LuminTargetPlatform.cpp b/Engine/Source/Developer/Lumin/LuminTargetPlatform/Private/LuminTargetPlatform.cpp
--- a/Engine/Source/Developer/Lumin/LuminTargetPlatform/Private/LuminTargetPlatform.cpp
+++ b/Engine/Source/Developer/Lumin/LuminTargetPlatform/Private/LuminTargetPlatform.cpp
@@ -103,9 +103,11 @@ void FLuminTargetPlatform::RefreshSettings()
 
 	// If we are targeting ES 2.0/3.1, we also must cook encoded HDR reflection captures
 	static FName NAME_VULKAN_ES31(TEXT("SF_VULKAN_ES31_LUMIN"));
+	static FName NAME_VULKAN_ES31_NOUB(TEXT("SF_VULKAN_ES31_LUMIN_NOUB"));
 	static FName NAME_GLSL_ES2(TEXT("GLSL_ES2"));
 	static FName NAME_GLSL_SM5(TEXT("GLSL_430"));
 	bRequiresEncodedHDRReflectionCaptures = TargetedShaderFormats.Contains(NAME_VULKAN_ES31)
+		|| TargetedShaderFormats.Contains(NAME_VULKAN_ES31_NOUB)
 		|| TargetedShaderFormats.Contains(NAME_GLSL_ES2)
 		|| TargetedShaderFormats.Contains(NAME_GLSL_SM5);
 
@@ -180,13 +182,18 @@ void FLuminTargetPlatform::GetAllPossibleShaderFormats( TArray<FName>& OutFormat
 //	static FName NAME_GLSL_SM4(TEXT("GLSL_150"));
 	static FName NAME_GLSL_SM5(TEXT("GLSL_430"));
 	static FName NAME_VULKAN_SM5_LUMIN(TEXT("SF_VULKAN_SM5_LUMIN"));
+	static FName NAME_VULKAN_SM5_LUMIN_NOUB(TEXT("SF_VULKAN_SM5_LUMIN_NOUB"));
 	static FName NAME_VULKAN_ES31_LUMIN(TEXT("SF_VULKAN_ES31_LUMIN"));
+	static FName NAME_VULKAN_ES31_LUMIN_NOUB(TEXT("SF_VULKAN_ES31_LUMIN_NOUB"));
+
+	static auto* CVar = IConsoleManager::Get().FindTConsoleVariableDataInt(TEXT("r.Vulkan.UseRealUBs"));
+	const bool bUseNOUB = (CVar && CVar->GetValueOnAnyThread() == 0);
 
 	if (SupportsMobileRendering())
 	{
 		if (LuminSupportsVulkan(LuminEngineSettings))
 		{
-			OutFormats.AddUnique(NAME_VULKAN_ES31_LUMIN);
+			OutFormats.AddUnique(bUseNOUB ? NAME_VULKAN_ES31_LUMIN_NOUB : NAME_VULKAN_ES31_LUMIN);
 		}
 		else
 		{
@@ -198,7 +205,7 @@ void FLuminTargetPlatform::GetAllPossibleShaderFormats( TArray<FName>& OutFormat
 	{
 		if (LuminSupportsVulkan(LuminEngineSettings))
 		{
-			OutFormats.AddUnique(NAME_VULKAN_SM5_LUMIN);
+			OutFormats.AddUnique(bUseNOUB ? NAME_VULKAN_SM5_LUMIN_NOUB : NAME_VULKAN_SM5_LUMIN);
 		}
 		else
 		{
diff --git a/Engine/Source/Developer/VulkanShaderFormat/Private/VulkanShaderFormat.cpp b/Engine/Source/Developer/VulkanShaderFormat/Private/VulkanShaderFormat.cpp
--- a/Engine/Source/Developer/VulkanShaderFormat/Private/VulkanShaderFormat.cpp
+++ b/Engine/Source/Developer/VulkanShaderFormat/Private/VulkanShaderFormat.cpp
@@ -13,12 +13,14 @@ static FName NAME_VULKAN_ES3_1_ANDROID(TEXT("SF_VULKAN_ES31_ANDROID"));
 static FName NAME_VULKAN_ES3_1_ANDROID_NOUB(TEXT("SF_VULKAN_ES31_ANDROID_NOUB"));
 static FName NAME_VULKAN_ES3_1(TEXT("SF_VULKAN_ES31"));
 static FName NAME_VULKAN_ES3_1_LUMIN(TEXT("SF_VULKAN_ES31_LUMIN"));
+static FName NAME_VULKAN_ES3_1_LUMIN_NOUB(TEXT("SF_VULKAN_ES31_LUMIN_NOUB"));
 static FName NAME_VULKAN_ES3_1_NOUB(TEXT("SF_VULKAN_ES31_NOUB"));
 static FName NAME_VULKAN_SM4_NOUB(TEXT("SF_VULKAN_SM4_NOUB"));
 static FName NAME_VULKAN_SM4(TEXT("SF_VULKAN_SM4"));
 static FName NAME_VULKAN_SM5_NOUB(TEXT("SF_VULKAN_SM5_NOUB"));
 static FName NAME_VULKAN_SM5(TEXT("SF_VULKAN_SM5"));
 static FName NAME_VULKAN_SM5_LUMIN(TEXT("SF_VULKAN_SM5_LUMIN"));
+static FName NAME_VULKAN_SM5_LUMIN_NOUB(TEXT("SF_VULKAN_SM5_LUMIN_NOUB"));
 
 class FShaderFormatVulkan : public IShaderFormat
 {
@@ -34,11 +36,11 @@ class FShaderFormatVulkan : public IShaderFormat
 		{
 			return UE_SHADER_VULKAN_SM5_VER;
 		}
-		else if (Format == NAME_VULKAN_SM5 || Format == NAME_VULKAN_SM5_NOUB || Format == NAME_VULKAN_SM5_LUMIN)
+		else if (Format == NAME_VULKAN_SM5 || Format == NAME_VULKAN_SM5_NOUB || Format == NAME_VULKAN_SM5_LUMIN || Format == NAME_VULKAN_SM5_LUMIN_NOUB)
 		{
 			return UE_SHADER_VULKAN_SM5_VER;
 		}
-		else if (Format == NAME_VULKAN_ES3_1_ANDROID || Format == NAME_VULKAN_ES3_1_ANDROID_NOUB || Format == NAME_VULKAN_ES3_1 || Format == NAME_VULKAN_ES3_1_NOUB || Format == NAME_VULKAN_ES3_1_LUMIN)
+		else if (Format == NAME_VULKAN_ES3_1_ANDROID || Format == NAME_VULKAN_ES3_1_ANDROID_NOUB || Format == NAME_VULKAN_ES3_1 || Format == NAME_VULKAN_ES3_1_NOUB || Format == NAME_VULKAN_ES3_1_LUMIN || Format == NAME_VULKAN_ES3_1_LUMIN_NOUB)
 		{
 			return UE_SHADER_VULKAN_ES3_1_VER;
 		}
@@ -62,10 +64,12 @@ public:
 		OutFormats.Add(NAME_VULKAN_SM4);
 		OutFormats.Add(NAME_VULKAN_SM5);
 		OutFormats.Add(NAME_VULKAN_SM5_LUMIN);
+		OutFormats.Add(NAME_VULKAN_SM5_LUMIN_NOUB);
 		OutFormats.Add(NAME_VULKAN_ES3_1_ANDROID);
 		OutFormats.Add(NAME_VULKAN_ES3_1_ANDROID_NOUB);
 		OutFormats.Add(NAME_VULKAN_ES3_1);
 		OutFormats.Add(NAME_VULKAN_ES3_1_LUMIN);
+		OutFormats.Add(NAME_VULKAN_ES3_1_LUMIN_NOUB);
 		OutFormats.Add(NAME_VULKAN_ES3_1_NOUB);
 		OutFormats.Add(NAME_VULKAN_SM4_NOUB);
 		OutFormats.Add(NAME_VULKAN_SM5_NOUB);
@@ -78,7 +82,7 @@ public:
 		{
 			DoCompileVulkanShader(Input, Output, WorkingDirectory, EVulkanShaderVersion::ES3_1);
 		}
-		else if (Format == NAME_VULKAN_ES3_1_NOUB)
+		else if (Format == NAME_VULKAN_ES3_1_NOUB || Format == NAME_VULKAN_ES3_1_LUMIN_NOUB)
 		{
 			DoCompileVulkanShader(Input, Output, WorkingDirectory, EVulkanShaderVersion::ES3_1_NOUB);
 		}
@@ -98,7 +102,7 @@ public:
 		{
 			DoCompileVulkanShader(Input, Output, WorkingDirectory, EVulkanShaderVersion::SM4);
 		}
-		else if (Format == NAME_VULKAN_SM5_NOUB)
+		else if (Format == NAME_VULKAN_SM5_NOUB || Format == NAME_VULKAN_SM5_LUMIN_NOUB)
 		{
 			DoCompileVulkanShader(Input, Output, WorkingDirectory, EVulkanShaderVersion::SM5_NOUB);
 		}
diff --git a/Engine/Source/Runtime/RHI/Public/RHIShaderFormatDefinitions.inl b/Engine/Source/Runtime/RHI/Public/RHIShaderFormatDefinitions.inl
--- a/Engine/Source/Runtime/RHI/Public/RHIShaderFormatDefinitions.inl
+++ b/Engine/Source/Runtime/RHI/Public/RHIShaderFormatDefinitions.inl
@@ -44,6 +44,7 @@ static FName NAME_SF_METAL_MACES2(TEXT("SF_METAL_MACES2"));
 static FName NAME_VULKAN_ES3_1_ANDROID(TEXT("SF_VULKAN_ES31_ANDROID"));
 static FName NAME_VULKAN_ES3_1_ANDROID_NOUB(TEXT("SF_VULKAN_ES31_ANDROID_NOUB"));
 static FName NAME_VULKAN_ES3_1_LUMIN(TEXT("SF_VULKAN_ES31_LUMIN"));
+static FName NAME_VULKAN_ES3_1_LUMIN_NOUB(TEXT("SF_VULKAN_ES31_LUMIN_NOUB"));
 static FName NAME_VULKAN_ES3_1(TEXT("SF_VULKAN_ES31"));
 static FName NAME_VULKAN_ES3_1_NOUB(TEXT("SF_VULKAN_ES31_NOUB"));
 static FName NAME_VULKAN_SM4_NOUB(TEXT("SF_VULKAN_SM4_NOUB"));
@@ -51,6 +52,7 @@ static FName NAME_VULKAN_SM4(TEXT("SF_VULKAN_SM4"));
 static FName NAME_VULKAN_SM5_NOUB(TEXT("SF_VULKAN_SM5_NOUB"));
 static FName NAME_VULKAN_SM5(TEXT("SF_VULKAN_SM5"));
 static FName NAME_VULKAN_SM5_LUMIN(TEXT("SF_VULKAN_SM5_LUMIN"));
+static FName NAME_VULKAN_SM5_LUMIN_NOUB(TEXT("SF_VULKAN_SM5_LUMIN_NOUB"));
 
 
 static FName ShaderPlatformToShaderFormatName(EShaderPlatform Platform)
@@ -123,7 +125,10 @@ static FName ShaderPlatformToShaderFormatName(EShaderPlatform Platform)
 		return NAME_VULKAN_ES3_1_ANDROID_NOUB;//NAME_VULKAN_ES3_1_ANDROID;
 
 	case SP_VULKAN_ES3_1_LUMIN:
-		return NAME_VULKAN_ES3_1_LUMIN;
+	{
+		static auto* CVar = IConsoleManager::Get().FindTConsoleVariableDataInt(TEXT("r.Vulkan.UseRealUBs"));
+		return (CVar && CVar->GetValueOnAnyThread() == 0) ? NAME_VULKAN_ES3_1_LUMIN_NOUB : NAME_VULKAN_ES3_1_LUMIN;
+	}
 
 	case SP_VULKAN_PCES3_1:
 	{
@@ -141,7 +146,10 @@ static FName ShaderPlatformToShaderFormatName(EShaderPlatform Platform)
 		return (CVar && CVar->GetValueOnAnyThread() == 0) ? NAME_VULKAN_SM5_NOUB : NAME_VULKAN_SM5;
 	}
 	case SP_VULKAN_SM5_LUMIN:
-		return NAME_VULKAN_SM5_LUMIN;
+	{
+		static auto* CVar = IConsoleManager::Get().FindTConsoleVariableDataInt(TEXT("r.Vulkan.UseRealUBs"));
+		return (CVar && CVar->GetValueOnAnyThread() == 0) ? NAME_VULKAN_SM5_LUMIN_NOUB : NAME_VULKAN_SM5_LUMIN;
+	}
 
 	default:
 		checkf(0, TEXT("Unknown EShaderPlatform %d!"), (int32)Platform);
@@ -187,6 +195,7 @@ static EShaderPlatform ShaderFormatNameToShaderPlatform(FName ShaderFormat)
 	if (ShaderFormat == NAME_VULKAN_ES3_1_ANDROID)		return SP_VULKAN_ES3_1_ANDROID;
 	if (ShaderFormat == NAME_VULKAN_ES3_1_ANDROID_NOUB)	return SP_VULKAN_ES3_1_ANDROID;
 	if (ShaderFormat == NAME_VULKAN_ES3_1_LUMIN)		return SP_VULKAN_ES3_1_LUMIN;
+	if (ShaderFormat == NAME_VULKAN_ES3_1_LUMIN_NOUB)	return SP_VULKAN_ES3_1_LUMIN;
 	if (ShaderFormat == NAME_VULKAN_ES3_1)				return SP_VULKAN_PCES3_1;
 	if (ShaderFormat == NAME_VULKAN_ES3_1_NOUB)			return SP_VULKAN_PCES3_1;
 	if (ShaderFormat == NAME_VULKAN_SM4_NOUB)			return SP_VULKAN_SM4;
@@ -194,6 +203,7 @@ static EShaderPlatform ShaderFormatNameToShaderPlatform(FName ShaderFormat)
 	if (ShaderFormat == NAME_VULKAN_SM5_NOUB)			return SP_VULKAN_SM5;
 	if (ShaderFormat == NAME_VULKAN_SM5)				return SP_VULKAN_SM5;
 	if (ShaderFormat == NAME_VULKAN_SM5_LUMIN)			return SP_VULKAN_SM5_LUMIN;
+	if (ShaderFormat == NAME_VULKAN_SM5_LUMIN_NOUB)		return SP_VULKAN_SM5_LUMIN;
 
 	return SP_NumPlatforms;
 }
diff --git a/Engine/Source/Runtime/UtilityShaders/Private/OneColorShader.cpp b/Engine/Source/Runtime/UtilityShaders/Private/OneColorShader.cpp
--- a/Engine/Source/Runtime/UtilityShaders/Private/OneColorShader.cpp
+++ b/Engine/Source/Runtime/UtilityShaders/Private/OneColorShader.cpp
@@ -3,6 +3,9 @@
 
 #include "OneColorShader.h"
 
+PRAGMA_DISABLE_OPTIMIZATION
+PRAGMA_DISABLE_INLINING
+
 BEGIN_GLOBAL_SHADER_PARAMETER_STRUCT(FClearShaderUB, )
 	SHADER_PARAMETER_ARRAY(FVector4, DrawColorMRT, [8] )
 END_GLOBAL_SHADER_PARAMETER_STRUCT()
diff --git a/Engine/Source/Runtime/VulkanRHI/Private/VulkanCommands.cpp b/Engine/Source/Runtime/VulkanRHI/Private/VulkanCommands.cpp
--- a/Engine/Source/Runtime/VulkanRHI/Private/VulkanCommands.cpp
+++ b/Engine/Source/Runtime/VulkanRHI/Private/VulkanCommands.cpp
@@ -10,6 +10,9 @@
 #include "EngineGlobals.h"
 #include "VulkanLLM.h"
 
+PRAGMA_DISABLE_OPTIMIZATION
+PRAGMA_DISABLE_INLINING
+
 static TAutoConsoleVariable<int32> GCVarSubmitOnDispatch(
 	TEXT("r.Vulkan.SubmitOnDispatch"),
 	0,
diff --git a/Engine/Source/Runtime/VulkanRHI/Private/VulkanRHI.cpp b/Engine/Source/Runtime/VulkanRHI/Private/VulkanRHI.cpp
--- a/Engine/Source/Runtime/VulkanRHI/Private/VulkanRHI.cpp
+++ b/Engine/Source/Runtime/VulkanRHI/Private/VulkanRHI.cpp
@@ -19,6 +19,9 @@
 #include "Misc/EngineVersion.h"
 #include "GlobalShader.h"
 
+PRAGMA_DISABLE_OPTIMIZATION
+PRAGMA_DISABLE_INLINING
+
 extern RHI_API bool GUseTexture3DBulkDataRHI;
 
 TAtomic<uint64> GVulkanBufferHandleIdCounter{ 0 };
@@ -48,7 +51,7 @@ static_assert(VK_API_VERSION >= UE_VK_API_VERSION, "Vulkan SDK is older than the
 
 TAutoConsoleVariable<int32> GRHIThreadCvar(
 	TEXT("r.Vulkan.RHIThread"),
-	1,
+	0,
 	TEXT("0 to only use Render Thread\n")
 	TEXT("1 to use ONE RHI Thread\n")
 	TEXT("2 to use multiple RHI Thread\n")
diff --git a/Engine/Source/Runtime/VulkanRHI/Public/VulkanGlobals.h b/Engine/Source/Runtime/VulkanRHI/Public/VulkanGlobals.h
--- a/Engine/Source/Runtime/VulkanRHI/Public/VulkanGlobals.h
+++ b/Engine/Source/Runtime/VulkanRHI/Public/VulkanGlobals.h
@@ -25,5 +25,5 @@ static FORCEINLINE void ZeroVulkanStruct(T& Struct, VkStructureType Type)
 
 inline bool UseVulkanDescriptorCache()
 {
-	return PLATFORM_ANDROID || GMaxRHIFeatureLevel <= ERHIFeatureLevel::ES3_1;
+	return (PLATFORM_ANDROID && !PLATFORM_LUMIN)|| GMaxRHIFeatureLevel <= ERHIFeatureLevel::ES3_1;
 }
