UE4.22 - Modify SPIR-V shader entry point name so it can be found easily while debugging (like Metal)
#rnx

diff --git a/Engine/Shaders/Public/Platform/Vulkan/VulkanCommon.ush b/Engine/Shaders/Public/Platform/Vulkan/VulkanCommon.ush
--- a/Engine/Shaders/Public/Platform/Vulkan/VulkanCommon.ush
+++ b/Engine/Shaders/Public/Platform/Vulkan/VulkanCommon.ush
@@ -5,7 +5,7 @@
 =============================================================================*/
 
 // Update this GUID to improve shader recompilation for Vulkan only shaders
-// GUID = {792BDC0A-A4A8-402F-8752-000000000040}
+// GUID = {792BDC0A-A4A8-402F-8752-000000000041}
 // Glslang 1.1.82.1
 
 #pragma once
diff --git a/Engine/Source/Developer/VulkanShaderFormat/Private/VulkanBackend.cpp b/Engine/Source/Developer/VulkanShaderFormat/Private/VulkanBackend.cpp
--- a/Engine/Source/Developer/VulkanShaderFormat/Private/VulkanBackend.cpp
+++ b/Engine/Source/Developer/VulkanShaderFormat/Private/VulkanBackend.cpp
@@ -5655,7 +5655,8 @@ bool FVulkanCodeBackend::GenerateMain(
 			_mesa_glsl_warning(ParseState, "'patchconstantfunc' attribute only applies to hull shaders");
 		}
 
-		ir_function* MainFunction = new(ParseState)ir_function("main");
+		// Values that will be patched in later from the SPIRV
+		ir_function* MainFunction = new(ParseState)ir_function("main_00000000_00000000");
 		MainFunction->add_signature(MainSig);
 
 		Instructions->append_list(&DeclInstructions);
diff --git a/Engine/Source/Developer/VulkanShaderFormat/Private/VulkanGlslangBridge.cpp b/Engine/Source/Developer/VulkanShaderFormat/Private/VulkanGlslangBridge.cpp
--- a/Engine/Source/Developer/VulkanShaderFormat/Private/VulkanGlslangBridge.cpp
+++ b/Engine/Source/Developer/VulkanShaderFormat/Private/VulkanGlslangBridge.cpp
@@ -208,10 +208,25 @@ static void ComputeMovableWordIndices(FSpirv& Spirv)
 		spv::Op OpCode = (spv::Op)(*Ptr & spv::OpCodeMask);
 		switch (OpCode)
 		{
+		case spv::OpEntryPoint:
+			{
+				uint32 ExecModel = Ptr[1];
+				uint32 EntryPoint = Ptr[2];
+				FString Name = ReadLiteralString(Ptr + 3);
+				check(Name == TEXT("main_00000000_00000000"));
+				check(Spirv.OffsetToEntryPoint == 0);
+				Spirv.OffsetToEntryPoint = (uint32)(&Ptr[3] - PtrStart);
+			}
+			break;
 		case spv::OpName:
 			{
 				uint32 TargetId = Ptr[1];
 				FString Name = ReadLiteralString(Ptr + 2);
+				if (Name == TEXT("main_00000000_00000000"))
+				{
+					check(Spirv.OffsetToMainName == 0);
+					Spirv.OffsetToMainName = (uint32)(&Ptr[2] - PtrStart);
+				}
 				Names.Add(TargetId, Name);
 			}
 			break;
@@ -328,6 +343,9 @@ bool GenerateSpirv(const ANSICHAR* Source, FCompilerInfo& CompilerInfo, FString&
 	const char* GlslSourceSkipHeader = strstr(Source, "#version");
 	Shader->setStrings(&GlslSourceSkipHeader, 1);
 
+	Shader->setEntryPoint("main_00000000_00000000");
+	Shader->setSourceEntryPoint("main_00000000_00000000");
+
 	auto DoGenerate = [&]()
 	{
 		const int DefaultVersion = 100;// Options & EOptionDefaultDesktop ? 110 : 100;
@@ -359,6 +377,7 @@ bool GenerateSpirv(const ANSICHAR* Source, FCompilerInfo& CompilerInfo, FString&
 			return false;
 		}
 
+		// Generate SPIRV
 		std::vector<unsigned int> Spirv;
 		glslang::GlslangToSpv(*Program->getIntermediate((EShLanguage)Stage), Spirv);
 
@@ -366,34 +385,7 @@ bool GenerateSpirv(const ANSICHAR* Source, FCompilerInfo& CompilerInfo, FString&
 		OutSpirv.Data.AddZeroed(SizeInWords);
 		FMemory::Memcpy(OutSpirv.Data.GetData(), &Spirv[0], SizeInWords * sizeof(uint32));
 
-		if (CompilerInfo.bDebugDump)
-		{
-			// Binary SpirV
-			FString SpirvFile = DumpDebugInfoPath / (TEXT("Output.spv"));
-			glslang::OutputSpvBin(Spirv, TCHAR_TO_ANSI(*SpirvFile));
-
-			// Text spirv
-			FString SpirvTextFile = DumpDebugInfoPath / (TEXT("Output.spvasm"));
-			std::ofstream File;
-			File.open(TCHAR_TO_ANSI(*SpirvTextFile), std::fstream::out/* | std::fstream::binary*/);
-			if (File.is_open())
-			{
-				spv::Parameterize();
-				spv::Disassemble(File, Spirv);
-				File.close();
-			}
-		}
-		else if (CompilerInfo.Input.bSkipPreprocessedCache)
-		{
-			spv::Parameterize();
-			spv::Disassemble(std::cout, Spirv);
-		}
-
-		if (CompilerInfo.Input.bSkipPreprocessedCache)
-		{
-			Program->dumpReflection();
-		}
-
+		// Gather Reflection info
 		for (int32 Index = 0; Index < Program->getNumLiveUniformVariables(); ++Index)
 		{
 			const char* AnsiName = Program->getUniformName(Index);
@@ -431,6 +423,48 @@ bool GenerateSpirv(const ANSICHAR* Source, FCompilerInfo& CompilerInfo, FString&
 		}
 
 		ComputeMovableWordIndices(OutSpirv);
+		OutSpirv.CRC = FCrc::MemCrc32(OutSpirv.Data.GetData(), OutSpirv.Data.Num() * sizeof(uint32));
+
+		// Patch the entry point name
+		auto FixEntryPoint = [&](uint32 OffsetToName)
+		{
+			char* EntryPointName = (char*)(OutSpirv.Data.GetData() + OffsetToName);
+			check(!FCStringAnsi::Strcmp(EntryPointName, "main_00000000_00000000"));
+			FCStringAnsi::Sprintf(EntryPointName, "main_%0.8x_%0.8x", OutSpirv.Data.Num() * sizeof(uint32), OutSpirv.CRC);
+		};
+		FixEntryPoint(OutSpirv.OffsetToMainName);
+		FixEntryPoint(OutSpirv.OffsetToEntryPoint);
+		// Copy back to original spirv data as it is used for dumping information
+		FMemory::Memcpy(&Spirv[0], OutSpirv.Data.GetData(), SizeInWords * sizeof(uint32));
+
+		// Dump debug info/files
+		if (CompilerInfo.bDebugDump)
+		{
+			// Binary SpirV
+			FString SpirvFile = DumpDebugInfoPath / (TEXT("Output.spv"));
+			glslang::OutputSpvBin(Spirv, TCHAR_TO_ANSI(*SpirvFile));
+
+			// Text spirv
+			FString SpirvTextFile = DumpDebugInfoPath / (TEXT("Output.spvasm"));
+			std::ofstream File;
+			File.open(TCHAR_TO_ANSI(*SpirvTextFile), std::fstream::out/* | std::fstream::binary*/);
+			if (File.is_open())
+			{
+				spv::Parameterize();
+				spv::Disassemble(File, Spirv);
+				File.close();
+			}
+		}
+		else if (CompilerInfo.Input.bSkipPreprocessedCache)
+		{
+			spv::Parameterize();
+			spv::Disassemble(std::cout, Spirv);
+		}
+
+		if (CompilerInfo.Input.bSkipPreprocessedCache)
+		{
+			Program->dumpReflection();
+		}
 
 		return true;
 	};
diff --git a/Engine/Source/Developer/VulkanShaderFormat/Private/VulkanShaderCompiler.cpp b/Engine/Source/Developer/VulkanShaderFormat/Private/VulkanShaderCompiler.cpp
--- a/Engine/Source/Developer/VulkanShaderFormat/Private/VulkanShaderCompiler.cpp
+++ b/Engine/Source/Developer/VulkanShaderFormat/Private/VulkanShaderCompiler.cpp
@@ -1027,6 +1027,7 @@ static void ConvertToNEWHeader(FOLDVulkanCodeHeader& OLDHeader,
 	OutHeader.EmulatedUBsCopyInfo = OLDHeader.UniformBuffersCopyInfo;
 	OutHeader.EmulatedUBCopyRanges = OLDHeader.NEWEmulatedUBCopyRanges;
 	OutHeader.SourceHash = OLDHeader.SourceHash;
+	OutHeader.SpirvCRC = Spirv.CRC;
 #if VULKAN_ENABLE_SHADER_DEBUG_NAMES
 	OutHeader.DebugName = OLDHeader.ShaderName;
 #endif
@@ -1041,8 +1042,7 @@ static void BuildShaderOutput(
 	const ANSICHAR* InShaderSource,
 	int32 SourceLen,
 	const FVulkanBindingTable& BindingTable,
-	const ANSICHAR* InShaderSourceES,
-	int32 SourceLenES,
+	uint32 NumLines,
 	FSpirv& Spirv,
 	const FString& DebugName,
 	bool bHasRealUBs
@@ -1489,7 +1489,8 @@ static void BuildShaderOutput(
 	//			uncommenting this will cause the project to have non deterministic materials and will hurt patch sizes
 	// ShaderOutput.ShaderCode.AddOptionalData('n', TCHAR_TO_UTF8(*ShaderInput.GenerateShaderName()));
 
-	ShaderOutput.NumInstructions = 0;
+	// Something to compare.
+	ShaderOutput.NumInstructions = NumLines;
 	ShaderOutput.NumTextureSamplers = OLDHeader.SerializedBindings.NumSamplers;
 	ShaderOutput.bSucceeded = true;
 
@@ -1514,7 +1515,6 @@ static void BuildShaderOutput(
 //	const ANSICHAR* InShaderSource,
 //	int32 SourceLen,
 //	const FVulkanBindingTable& BindingTable,
-//	const ANSICHAR* InShaderSourceES,
 //	int32 SourceLenES,
 //	const FString& SPVFile,
 //	const FString& DebugName
@@ -1529,7 +1529,6 @@ static void BuildShaderOutput(
 //		InShaderSource,
 //		SourceLen,
 //		BindingTable,
-//		InShaderSourceES,
 //		SourceLenES,
 //		Spirv,
 //		DebugName
@@ -1731,6 +1730,24 @@ static bool CompileUsingInternal(FCompilerInfo& CompilerInfo, FVulkanBindingTabl
 {
 	FString Errors;
 	FSpirv Spirv;
+	const ANSICHAR* Main = GlslSource.GetData();
+	Main = FCStringAnsi::Strstr(Main, "void main_");
+	check(Main);
+	auto GetNumEOLs = [](const ANSICHAR* Ptr)
+	{
+		uint32 NumLines = 0;
+		while (*Ptr)
+		{
+			if (*Ptr == '\n')
+			{
+				++NumLines;
+			}
+			++Ptr;
+		}
+
+		return NumLines;
+	};
+	uint32 NumLines = GetNumEOLs(Main);
 	if (GenerateSpirv(GlslSource.GetData(), CompilerInfo, Errors, CompilerInfo.Input.DumpDebugInfoPath, Spirv))
 	{
 		FString DebugName = CompilerInfo.Input.DumpDebugInfoPath.Right(CompilerInfo.Input.DumpDebugInfoPath.Len() - CompilerInfo.Input.DumpDebugInfoRootPath.Len());
@@ -1738,7 +1755,7 @@ static bool CompileUsingInternal(FCompilerInfo& CompilerInfo, FVulkanBindingTabl
 		Output.Target = CompilerInfo.Input.Target;
 		BuildShaderOutput(Output, CompilerInfo.Input,
 			GlslSource.GetData(), GlslSource.Num(),
-			BindingTable, nullptr, 0, Spirv, DebugName, bHasRealUBs);
+			BindingTable, NumLines, Spirv, DebugName, bHasRealUBs);
 		return true;
 	}
 	else
diff --git a/Engine/Source/Developer/VulkanShaderFormat/Private/VulkanShaderFormat.cpp b/Engine/Source/Developer/VulkanShaderFormat/Private/VulkanShaderFormat.cpp
--- a/Engine/Source/Developer/VulkanShaderFormat/Private/VulkanShaderFormat.cpp
+++ b/Engine/Source/Developer/VulkanShaderFormat/Private/VulkanShaderFormat.cpp
@@ -26,8 +26,8 @@ class FShaderFormatVulkan : public IShaderFormat
 {
 	enum 
 	{
-		UE_SHADER_VULKAN_ES3_1_VER	= 25,
-		UE_SHADER_VULKAN_SM5_VER 	= 25,
+		UE_SHADER_VULKAN_ES3_1_VER	= 26,
+		UE_SHADER_VULKAN_SM5_VER 	= 26,
 	};
 
 	int32 InternalGetVersion(FName Format) const
diff --git a/Engine/Source/Developer/VulkanShaderFormat/Private/VulkanShaderFormat.h b/Engine/Source/Developer/VulkanShaderFormat/Private/VulkanShaderFormat.h
--- a/Engine/Source/Developer/VulkanShaderFormat/Private/VulkanShaderFormat.h
+++ b/Engine/Source/Developer/VulkanShaderFormat/Private/VulkanShaderFormat.h
@@ -81,6 +81,9 @@ struct FSpirv
 		uint32 WordBindingIndex = UINT32_MAX;
 	};
 	TArray<FEntry> ReflectionInfo;
+	uint32 OffsetToMainName = 0;
+	uint32 OffsetToEntryPoint = 0;
+	uint32 CRC = 0;
 
 	int32 FindBinding(const FString& Name, bool bOuter = false) const
 	{
diff --git a/Engine/Source/Runtime/VulkanRHI/Private/VulkanPipeline.cpp b/Engine/Source/Runtime/VulkanRHI/Private/VulkanPipeline.cpp
--- a/Engine/Source/Runtime/VulkanRHI/Private/VulkanPipeline.cpp
+++ b/Engine/Source/Runtime/VulkanRHI/Private/VulkanPipeline.cpp
@@ -1120,6 +1120,8 @@ void FVulkanPipelineStateCacheManager::CreateGfxPipelineFromEntry(FGfxPipelineEn
 	FMemory::Memzero(ShaderStages);
 	PipelineInfo.stageCount = 0;
 	PipelineInfo.pStages = ShaderStages;
+	// main_00000000_00000000
+	ANSICHAR EntryPoints[ShaderStage::NumStages][24];
 	for (int32 ShaderStage = 0; ShaderStage < ShaderStage::NumStages; ++ShaderStage)
 	{
 		if (!GfxEntry->ShaderModules[ShaderStage])
@@ -1131,7 +1133,8 @@ void FVulkanPipelineStateCacheManager::CreateGfxPipelineFromEntry(FGfxPipelineEn
 		ShaderStages[PipelineInfo.stageCount].sType = VK_STRUCTURE_TYPE_PIPELINE_SHADER_STAGE_CREATE_INFO;
 		ShaderStages[PipelineInfo.stageCount].stage = UEFrequencyToVKStageBit(ShaderStage::GetFrequencyForGfxStage(CurrStage));
 		ShaderStages[PipelineInfo.stageCount].module = GfxEntry->ShaderModules[CurrStage];
-		ShaderStages[PipelineInfo.stageCount].pName = "main";
+		Shaders[ShaderStage]->GetEntryPoint(EntryPoints[PipelineInfo.stageCount]);
+		ShaderStages[PipelineInfo.stageCount].pName = EntryPoints[PipelineInfo.stageCount];
 		PipelineInfo.stageCount++;
 	}
 
@@ -1712,7 +1715,10 @@ FVulkanComputePipeline* FVulkanPipelineStateCacheManager::CreateComputePipelineF
 	PipelineInfo.stage.sType = VK_STRUCTURE_TYPE_PIPELINE_SHADER_STAGE_CREATE_INFO;
 	PipelineInfo.stage.stage = VK_SHADER_STAGE_COMPUTE_BIT;
 	PipelineInfo.stage.module = ShaderModule;
-	PipelineInfo.stage.pName = "main";
+	// main_00000000_00000000
+	ANSICHAR EntryPoint[24];
+	Shader->GetEntryPoint(EntryPoint);
+	PipelineInfo.stage.pName = EntryPoint;
 	PipelineInfo.layout = ComputeLayout->GetPipelineLayout();
 		
 	VERIFYVULKANRESULT(VulkanRHI::vkCreateComputePipelines(Device->GetInstanceHandle(), PipelineCache, 1, &PipelineInfo, VULKAN_CPU_ALLOCATOR, &Pipeline->Pipeline));
diff --git a/Engine/Source/Runtime/VulkanRHI/Public/VulkanResources.h b/Engine/Source/Runtime/VulkanRHI/Public/VulkanResources.h
--- a/Engine/Source/Runtime/VulkanRHI/Public/VulkanResources.h
+++ b/Engine/Source/Runtime/VulkanRHI/Public/VulkanResources.h
@@ -108,6 +108,12 @@ public:
 	}
 #endif
 
+	// Name should be pointing to "main_"
+	void GetEntryPoint(ANSICHAR* Name)
+	{
+		FCStringAnsi::Sprintf(Name, "main_%0.8x_%0.8x", Spirv.Num() * sizeof(uint32), CodeHeader.SpirvCRC);
+	}
+
 	FORCEINLINE const FVulkanShaderHeader& GetCodeHeader() const
 	{
 		return CodeHeader;
diff --git a/Engine/Source/Runtime/VulkanRHI/Public/VulkanShaderResources.h b/Engine/Source/Runtime/VulkanRHI/Public/VulkanShaderResources.h
--- a/Engine/Source/Runtime/VulkanRHI/Public/VulkanShaderResources.h
+++ b/Engine/Source/Runtime/VulkanRHI/Public/VulkanShaderResources.h
@@ -176,7 +176,7 @@ struct FVulkanShaderHeader
 	uint16									Pad1 = 1;
 
 	FSHAHash								SourceHash;
-	//FSHAHash								SpirvHash;
+	uint32									SpirvCRC = 0;
 
 	TArray<FSpirvInfo>						UniformBufferSpirvInfos;
 	TArray<FSpirvInfo>						GlobalSpirvInfos;
@@ -282,6 +282,7 @@ inline FArchive& operator<<(FArchive& Ar, FVulkanShaderHeader& Header)
 	Ar << Header.InOutMask;
 	Ar << Header.bHasRealUBs;
 	Ar << Header.SourceHash;
+	Ar << Header.SpirvCRC;
 	Ar << Header.UniformBufferSpirvInfos;
 	Ar << Header.GlobalSpirvInfos;
 #if VULKAN_ENABLE_SHADER_DEBUG_NAMES
