DR - vk - Zero-initialize all values in Vulkan shaders

diff --git a/Engine/Shaders/Public/Platform/Vulkan/VulkanCommon.ush b/Engine/Shaders/Public/Platform/Vulkan/VulkanCommon.ush
--- a/Engine/Shaders/Public/Platform/Vulkan/VulkanCommon.ush
+++ b/Engine/Shaders/Public/Platform/Vulkan/VulkanCommon.ush
@@ -5,7 +5,7 @@
 =============================================================================*/
 
 // Update this GUID to improve shader recompilation for Vulkan only shaders
-// GUID = {792BDC0A-A4A8-402F-8752-00000000003C}
+// GUID = {792BDC0A-A4A8-402F-8752-00000000003D}
 // Glslang 1.1.82.1
 
 #pragma once
diff --git a/Engine/Source/Developer/VulkanShaderFormat/Private/VulkanBackend.cpp b/Engine/Source/Developer/VulkanShaderFormat/Private/VulkanBackend.cpp
--- a/Engine/Source/Developer/VulkanShaderFormat/Private/VulkanBackend.cpp
+++ b/Engine/Source/Developer/VulkanShaderFormat/Private/VulkanBackend.cpp
@@ -752,6 +752,7 @@ class FGenerateVulkanVisitor : public ir_visitor
 	bool bUsesDXDY;
 
 	std::vector<std::string> SamplerStateNames;
+	TIRVarSet AtomicVariables;
 
 	/**
 	* Return true if the type is a multi-dimensional array. Also, track the
@@ -1029,6 +1030,55 @@ class FGenerateVulkanVisitor : public ir_visitor
 		check(0 && "ir_rvalue not handled for GLSL export.");
 	}
 
+	static inline bool is_struct_type(const glsl_type * type)
+	{
+		if (type->base_type != GLSL_TYPE_STRUCT && type->base_type != GLSL_TYPE_INPUTPATCH)
+		{
+			if (type->base_type == GLSL_TYPE_ARRAY && type->element_type())
+			{
+				return is_struct_type(type->element_type());
+			}
+			else
+			{
+				return false;
+			}
+		}
+		else
+		{
+			return true;
+		}
+	}
+
+	void print_zero_initialiser(const glsl_type * type)
+	{
+		if (type->base_type != GLSL_TYPE_STRUCT)
+		{
+			if (type->base_type != GLSL_TYPE_ARRAY)
+			{
+				ir_constant* zero = ir_constant::zero(mem_ctx, type);
+				if (zero)
+				{
+					zero->accept(this);
+				}
+			}
+			else
+			{
+				ralloc_asprintf_append(buffer, "{");
+
+				for (uint32 i = 0; i < type->length; i++)
+				{
+					if (i > 0)
+					{
+						ralloc_asprintf_append(buffer, ", ");
+					}
+					print_zero_initialiser(type->element_type());
+				}
+
+				ralloc_asprintf_append(buffer, "}");
+			}
+		}
+	}
+
 	virtual void visit(ir_variable *var)
 	{
 		const char * const centroid_str[] = { "", "centroid " };
@@ -1332,6 +1382,14 @@ class FGenerateVulkanVisitor : public ir_visitor
 				var->constant_value->accept(this);
 			}
 		}
+		else if ((var->type->base_type != GLSL_TYPE_STRUCT) && (var->mode == ir_var_auto || var->mode == ir_var_temporary || var->mode == ir_var_shared) && (AtomicVariables.find(var) == AtomicVariables.end()))
+		{
+			if (!is_struct_type(var->type) && var->type->base_type != GLSL_TYPE_ARRAY && var->mode != ir_var_shared)
+			{
+				ralloc_asprintf_append(buffer, " = ");
+				print_zero_initialiser(var->type);
+			}
+		}
 
 		// add type to used_structures so we can later declare them at the start of the GLSL shader
 		// this is for the case of a variable that is declared, but not later dereferenced (which can happen
@@ -3351,6 +3409,11 @@ public:
 		hash_table_dtor(used_uniform_blocks);
 	}
 
+	void FindAtomicVariables(exec_list* ir)
+	{
+		::FindAtomicVariables(ir, AtomicVariables);
+	}
+
 	int32 AddUniqueSamplerState(const std::string& Name)
 	{
 		if (Name == "")
@@ -3787,6 +3850,7 @@ char* FVulkanCodeBackend::GenerateCode(exec_list* ir, _mesa_glsl_parse_state* st
 
 	// Setup root visitor
 	FGenerateVulkanVisitor visitor(Target, BindingTable, state, bGenerateLayoutLocations, bDefaultPrecisionIsHalf);
+	visitor.FindAtomicVariables(ir);
 
 	// Generate information for sharing samplers
 	{
