DR - vk - Copy from 4580291
- Fix for unhandled nan in Vulkan shaders

diff --git a/Engine/Source/Developer/VulkanShaderFormat/Private/VulkanBackend.cpp b/Engine/Source/Developer/VulkanShaderFormat/Private/VulkanBackend.cpp
--- a/Engine/Source/Developer/VulkanShaderFormat/Private/VulkanBackend.cpp
+++ b/Engine/Source/Developer/VulkanShaderFormat/Private/VulkanBackend.cpp
@@ -2158,19 +2158,24 @@ class FGenerateVulkanVisitor : public ir_visitor
 				switch (constant->value.u[index])
 				{
 				case 0x7f800000u:
-					ralloc_asprintf_append(buffer, "(1.0/0.0)");
+					ralloc_asprintf_append(buffer, "(1.0/0.0) /*Inf*/");
 					break;
 
 				case 0xffc00000u:
-					ralloc_asprintf_append(buffer, "(0.0/0.0)");
+					ralloc_asprintf_append(buffer, "(0.0/0.0) /*-Nan*/");
 					break;
 
 				case 0xff800000u:
-					ralloc_asprintf_append(buffer, "(-1.0/0.0)");
+					ralloc_asprintf_append(buffer, "(-1.0/0.0) /*-Inf*/");
+					break;
+
+				case 0x7fc00000u:
+					ralloc_asprintf_append(buffer, "(0.0/0.0) /*Nan*/");
 					break;
 
 				default:
-					check(0);
+					checkf(false, TEXT("constant->value.u[index] = 0x%0x"), constant->value.u[index]);
+					break;
 				}
 			}
 		}
