DR - vk - Do not emit precision on bool variables
#rnx

diff --git a/Engine/Shaders/Public/Platform/Vulkan/VulkanCommon.ush b/Engine/Shaders/Public/Platform/Vulkan/VulkanCommon.ush
--- a/Engine/Shaders/Public/Platform/Vulkan/VulkanCommon.ush
+++ b/Engine/Shaders/Public/Platform/Vulkan/VulkanCommon.ush
@@ -5,7 +5,7 @@
 =============================================================================*/
 
 // Update this GUID to improve shader recompilation for Vulkan only shaders
-// GUID = {792BDC0A-A4A8-402F-8752-00000000003D}
+// GUID = {792BDC0A-A4A8-402F-8752-00000000003E}
 // Glslang 1.1.82.1
 
 #pragma once
diff --git a/Engine/Source/Developer/VulkanShaderFormat/Private/VulkanBackend.cpp b/Engine/Source/Developer/VulkanShaderFormat/Private/VulkanBackend.cpp
--- a/Engine/Source/Developer/VulkanShaderFormat/Private/VulkanBackend.cpp
+++ b/Engine/Source/Developer/VulkanShaderFormat/Private/VulkanBackend.cpp
@@ -960,6 +960,10 @@ class FGenerateVulkanVisitor : public ir_visitor
 
 	EPrecisionModifier GetPrecisionModifier(const struct glsl_type *type)
 	{
+		if (type->base_type == GLSL_TYPE_BOOL)
+		{
+			return GLSL_PRECISION_DEFAULT;
+		}
 		if (type->is_sampler() || type->is_image())
 		{
 			if (bDefaultPrecisionIsHalf && type->inner_type->base_type == GLSL_TYPE_FLOAT)
@@ -2777,8 +2781,8 @@ class FGenerateVulkanVisitor : public ir_visitor
 						//EHart - name-mangle variables to prevent colliding names
 						//#todo-rco: Check if this is still is needed when creating PSOs
 						//ralloc_asprintf_append(buffer, "#define %s %s%s\n", var->name, var->name, block_name);
-
-						ralloc_asprintf_append(buffer, "\t%s", (state->language_version == 310 && bEmitPrecision) ? "highp " : "");
+						bool bIsBoolType = var->type->base_type == GLSL_TYPE_BOOL;
+						ralloc_asprintf_append(buffer, "\t%s", (state->language_version == 310 && bEmitPrecision && !bIsBoolType) ? "highp " : "");
 						print_type_pre(var->type);
 						ralloc_asprintf_append(buffer, " %s", var->name);
 						print_type_post(var->type);
