Bump shader version for jira UE-71386.
Followup change to the vulkan nan fix.
hlslcc was not generating conditional moves for swizzles, causing the proposed fix to generate a lot of code. this is fixed with this change

Code Before:
{
		bvec4 v138 = bvec4(false,false,false,false);
		v138.xyzw = lessThan(floatBitsToUint(v137),uvec4(2139095040u,2139095040u,2139095040u,2139095040u));
		vec4 v139 = vec4(0.000000,0.000000,0.000000,0.000000);
		if (not(v138).x)
		{
			v139.x = 0.000000;
		}
		else
		{
			v139.x = v137.x;
		}
		if (not(v138).y)
		{
			v139.y = 0.000000;
		}
		else
		{
			v139.y = v137.y;
		}
		if (not(v138).z)
		{
			v139.z = 0.000000;
		}
		else
		{
			v139.z = v137.z;
		}
		if (not(v138).w)
		{
			v139.w = 0.000000;
		}
		else
		{
			v139.w = v137.w;
		}
		imageStore( RWLightScattering, ivec3(gl_GlobalInvocationID), v139.xyzw);
	}

Code After

	{
		bvec4 v138 = bvec4(false,false,false,false);
		v138.xyzw = lessThan(floatBitsToUint(v137),uvec4(2139095040u,2139095040u,2139095040u,2139095040u));
		vec4 v139 = vec4(0.000000,0.000000,0.000000,0.000000);
		v139.x = (not(v138).x)?(0.000000):(v137.x);
		v139.y = (not(v138).y)?(0.000000):(v137.y);
		v139.z = (not(v138).z)?(0.000000):(v137.z);
		v139.w = (not(v138).w)?(0.000000):(v137.w);
		imageStore( RWLightScattering, ivec3(gl_GlobalInvocationID), v139.xyzw);
	}

diff --git a/Engine/Shaders/Public/ShaderVersion.ush b/Engine/Shaders/Public/ShaderVersion.ush
--- a/Engine/Shaders/Public/ShaderVersion.ush
+++ b/Engine/Shaders/Public/ShaderVersion.ush
@@ -3,4 +3,4 @@
 // in Platform.ush (which should be included in any shader) it allows to invalidate the shader DDC.
 // 
 // If you are merging streams and there is a conflict with this GUID you should make a new GUID rather than taking one or the other.
-// GUID = B0A8821C-374F-4F39-92B6-CABBF1F4DA32
+// GUID = 91B6C8C8-2F23-41D9-B55D-AF3C0F910ED1
diff --git a/Engine/Source/Developer/VulkanShaderFormat/Private/VulkanBackend.cpp b/Engine/Source/Developer/VulkanShaderFormat/Private/VulkanBackend.cpp
--- a/Engine/Source/Developer/VulkanShaderFormat/Private/VulkanBackend.cpp
+++ b/Engine/Source/Developer/VulkanShaderFormat/Private/VulkanBackend.cpp
@@ -2344,7 +2344,7 @@ class FGenerateVulkanVisitor : public ir_visitor
 
 			ir_instruction *const inst = (ir_instruction *)iter.get();
 			ir_assignment *assignment = inst->as_assignment();
-			if (assignment && (assignment->rhs->ir_type == ir_type_dereference_variable || assignment->rhs->ir_type == ir_type_constant))
+			if (assignment && (assignment->rhs->ir_type == ir_type_dereference_variable || assignment->rhs->ir_type == ir_type_constant || assignment->rhs->ir_type == ir_type_swizzle))
 			{
 				dest_deref = assignment->lhs->as_dereference_variable();
 				true_value = assignment->rhs;
@@ -2367,7 +2367,7 @@ class FGenerateVulkanVisitor : public ir_visitor
 
 			ir_instruction *const inst = (ir_instruction *)iter.get();
 			ir_assignment *assignment = inst->as_assignment();
-			if (assignment && (assignment->rhs->ir_type == ir_type_dereference_variable || assignment->rhs->ir_type == ir_type_constant))
+			if (assignment && (assignment->rhs->ir_type == ir_type_dereference_variable || assignment->rhs->ir_type == ir_type_constant || assignment->rhs->ir_type == ir_type_swizzle))
 			{
 				ir_dereference_variable *tmp_deref = assignment->lhs->as_dereference_variable();
 				if (tmp_deref
