Add GL_OES_shader_image_atomic as a requirement for VK feature level ES3.1, to support Niagara GPU sim

diff --git a/Engine/Source/Developer/VulkanShaderFormat/Private/VulkanBackend.cpp b/Engine/Source/Developer/VulkanShaderFormat/Private/VulkanBackend.cpp
--- a/Engine/Source/Developer/VulkanShaderFormat/Private/VulkanBackend.cpp
+++ b/Engine/Source/Developer/VulkanShaderFormat/Private/VulkanBackend.cpp
@@ -3380,7 +3380,10 @@ class FGenerateVulkanVisitor : public ir_visitor
 			{
 				ralloc_asprintf_append(buffer, "#extension GL_EXT_tessellation_shader : enable\n");
 			}
-
+			else if (ParseState->target == compute_shader)
+			{
+				ralloc_asprintf_append(buffer, "#extension GL_OES_shader_image_atomic : enable\n");
+			}
 		}
 
 		if (bShouldEmitMultiview && (state->target == vertex_shader || state->target == fragment_shader))
