Cherry-pick 7106794: Added "GL_OES_shader_image_atomic" extension to Vulkan backend for Lumin platform.

#rnx

diff --git a/Engine/Shaders/Public/Platform/Vulkan/VulkanCommon.ush b/Engine/Shaders/Public/Platform/Vulkan/VulkanCommon.ush
--- a/Engine/Shaders/Public/Platform/Vulkan/VulkanCommon.ush
+++ b/Engine/Shaders/Public/Platform/Vulkan/VulkanCommon.ush
@@ -5,7 +5,7 @@
 =============================================================================*/
 
 // Update this GUID to improve shader recompilation for Vulkan only shaders
-// GUID = 86A78414-CF0D-48DB-AC9C-8A8301C857A5
+// GUID = 879F6C92-E759-47F4-A195-9F5FD9945012
 // Glslang 1.1.82.1
 
 #pragma once
diff --git a/Engine/Source/Developer/VulkanShaderFormat/Private/VulkanBackend.cpp b/Engine/Source/Developer/VulkanShaderFormat/Private/VulkanBackend.cpp
--- a/Engine/Source/Developer/VulkanShaderFormat/Private/VulkanBackend.cpp
+++ b/Engine/Source/Developer/VulkanShaderFormat/Private/VulkanBackend.cpp
@@ -769,9 +769,12 @@ class FGenerateVulkanVisitor : public ir_visitor
 	/** Whether the shader being cross compiled needs EXT_shader_texture_lod. */
 	bool bUsesES2TextureLODExtension;
 
-	// Found dFdx or dFdy
+	/** Found dFdx or dFdy */
 	bool bUsesDXDY;
 
+	/** Found image atomic functions (e.g. imageAtomicAdd) */
+	bool bUsesImageWriteAtomic;
+
 	std::vector<std::string> SamplerStateNames;
 	TIRVarSet AtomicVariables;
 
@@ -2576,6 +2579,8 @@ class FGenerateVulkanVisitor : public ir_visitor
 				ir->operands[1]->accept(this);
 			}
 			ralloc_asprintf_append(buffer, ")");
+
+			bUsesImageWriteAtomic = true;
 		}
 	}
 
@@ -3345,7 +3350,7 @@ class FGenerateVulkanVisitor : public ir_visitor
 #endif		
 	}
 
-	void print_extensions(_mesa_glsl_parse_state* state, bool bUsesES31Extensions, bool bShouldEmitMultiview)
+	void print_extensions(_mesa_glsl_parse_state* state, bool bUsesES31Extensions, bool bShouldEmitOESExtensions, bool bShouldEmitMultiview)
 	{
 		if (bUsesES2TextureLODExtension)
 		{
@@ -3365,6 +3370,11 @@ class FGenerateVulkanVisitor : public ir_visitor
 			ralloc_asprintf_append(buffer, "#extension GL_OES_standard_derivatives : enable\n");
 		}
 
+		if (bUsesImageWriteAtomic && bShouldEmitOESExtensions)
+		{
+			ralloc_asprintf_append(buffer, "#extension GL_OES_shader_image_atomic : enable\n");
+		}
+
 		if (bUsesES31Extensions)
 		{
 			ralloc_asprintf_append(buffer, "#extension GL_EXT_gpu_shader5 : enable\n");
@@ -3420,6 +3430,7 @@ public:
 		, loop_count(0)
 		, bUsesES2TextureLODExtension(false)
 		, bUsesDXDY(false)
+		, bUsesImageWriteAtomic(false)
 	{
 		printable_names = hash_table_ctor(32, hash_table_pointer_hash, hash_table_pointer_compare);
 		used_structures = hash_table_ctor(32, hash_table_pointer_hash, hash_table_pointer_compare);
@@ -3521,11 +3532,14 @@ public:
 		print_layout(state);
 		buffer = 0;
 
+		const FVulkanLanguageSpec* LanguageSpec = static_cast<const FVulkanLanguageSpec*>(state->LanguageSpec);
+
+		const bool bShouldEmitOESExtensions = LanguageSpec->RequiresOESExtensions();
 		const bool bShouldEmitMultiview = strstr(signature, "gl_ViewIndex") != nullptr;
 
 		char* Extensions = ralloc_asprintf(mem_ctx, "");
 		buffer = &Extensions;
-		print_extensions(state, state->language_version == 310, bShouldEmitMultiview);
+		print_extensions(state, state->language_version == 310, bShouldEmitOESExtensions, bShouldEmitMultiview);
 		if (state->bSeparateShaderObjects && !state->bGenerateES)
 		{
 			switch (state->target)
diff --git a/Engine/Source/Developer/VulkanShaderFormat/Private/VulkanBackend.h b/Engine/Source/Developer/VulkanShaderFormat/Private/VulkanBackend.h
--- a/Engine/Source/Developer/VulkanShaderFormat/Private/VulkanBackend.h
+++ b/Engine/Source/Developer/VulkanShaderFormat/Private/VulkanBackend.h
@@ -11,10 +11,12 @@ class FVulkanLanguageSpec : public ILanguageSpec
 {
 protected:
 	bool bShareSamplers;
+	bool bRequiresOESExtensions;
 
 public:
-	FVulkanLanguageSpec(bool bInShareSamplers)
+	FVulkanLanguageSpec(bool bInShareSamplers, bool bInRequiresOESExtensions)
 		: bShareSamplers(bInShareSamplers)
+		, bRequiresOESExtensions(bInRequiresOESExtensions)
 	{}
 
 	virtual bool SupportsDeterminantIntrinsic() const override
@@ -39,6 +41,9 @@ public:
 	virtual bool AllowsSharingSamplers() const override { return bShareSamplers; }
 
 	virtual bool RequiresNegateDDY() const override { return false; }
+
+public:
+	bool RequiresOESExtensions() const { return bRequiresOESExtensions; }
 };
 
 class ir_variable;
diff --git a/Engine/Source/Developer/VulkanShaderFormat/Private/VulkanShaderCompiler.cpp b/Engine/Source/Developer/VulkanShaderFormat/Private/VulkanShaderCompiler.cpp
--- a/Engine/Source/Developer/VulkanShaderFormat/Private/VulkanShaderCompiler.cpp
+++ b/Engine/Source/Developer/VulkanShaderFormat/Private/VulkanShaderCompiler.cpp
@@ -1781,8 +1781,11 @@ static bool CallHlslcc(const FString& PreprocessedShader, FVulkanBindingTable& B
 		// Call hlslcc
 		FVulkanCodeBackend VulkanBackend(CompilerInfo.CCFlags, BindingTable, HlslCompilerTarget);
 		FHlslCrossCompilerContext CrossCompilerContext(CompilerInfo.CCFlags, CompilerInfo.Frequency, HlslCompilerTarget);
-		const bool bShareSamplers = false;
-		FVulkanLanguageSpec VulkanLanguageSpec(true);
+
+		const bool bShareSamplers = true;
+		const bool bRequiresOESExtensions = (CompilerInfo.Input.Target.Platform == SP_VULKAN_ES3_1_LUMIN);
+
+		FVulkanLanguageSpec VulkanLanguageSpec(bShareSamplers, bRequiresOESExtensions);
 		int32 Result = 0;
 		if (CrossCompilerContext.Init(TCHAR_TO_ANSI(*CompilerInfo.Input.VirtualSourceFilePath), &VulkanLanguageSpec))
 		{
diff --git a/Engine/Source/Runtime/ApplicationCore/Private/Android/AndroidWindow.cpp b/Engine/Source/Runtime/ApplicationCore/Private/Android/AndroidWindow.cpp
--- a/Engine/Source/Runtime/ApplicationCore/Private/Android/AndroidWindow.cpp
+++ b/Engine/Source/Runtime/ApplicationCore/Private/Android/AndroidWindow.cpp
@@ -164,10 +164,10 @@ void FAndroidWindow::ReleaseWindowRef(ANativeWindow* InWindow)
 
  void FAndroidWindow::SetHardwareWindow_EventThread(void* InWindow)
 {
-#if USE_ANDROID_JNI
+#if USE_ANDROID_EVENTS
 	check(IsInAndroidEventThread());
-	NativeWindow = InWindow; //using raw native window handle for now. Could be changed to use AndroidWindow later if needed
 #endif
+	NativeWindow = InWindow; //using raw native window handle for now. Could be changed to use AndroidWindow later if needed
 }
 
 void* FAndroidWindow::GetHardwareWindow_EventThread()
