Applied same fix for Vulkan tessellation on AMD GPU as in Dev-Rendering (see CL 8129425).

diff --git a/Engine/Shaders/Public/Platform/Vulkan/VulkanCommon.ush b/Engine/Shaders/Public/Platform/Vulkan/VulkanCommon.ush
--- a/Engine/Shaders/Public/Platform/Vulkan/VulkanCommon.ush
+++ b/Engine/Shaders/Public/Platform/Vulkan/VulkanCommon.ush
@@ -5,7 +5,7 @@
 =============================================================================*/
 
 // Update this GUID to improve shader recompilation for Vulkan only shaders
-// GUID = 879F6C92-E759-47F4-A195-9F5FD9945012
+// GUID = A4B869BA-2A9A-40B8-B8ED-7FA3B6DFEEF7
 // Glslang 1.1.82.1
 
 #pragma once
diff --git a/Engine/Source/Developer/VulkanShaderFormat/Private/VulkanBackend.cpp b/Engine/Source/Developer/VulkanShaderFormat/Private/VulkanBackend.cpp
--- a/Engine/Source/Developer/VulkanShaderFormat/Private/VulkanBackend.cpp
+++ b/Engine/Source/Developer/VulkanShaderFormat/Private/VulkanBackend.cpp
@@ -3549,7 +3549,7 @@ public:
 				ralloc_asprintf_append(buffer, "in gl_PerVertex\n"
 					"{\n"
 					"\tvec4 gl_Position;\n"
-					"\tfloat gl_ClipDistance[6];\n"
+					"\tfloat gl_ClipDistance[];\n"
 					"} gl_in[];\n"
 					);
 #endif
@@ -3559,7 +3559,7 @@ public:
 				ralloc_asprintf_append(buffer, "out gl_PerVertex\n"
 					"{\n"
 					"\tvec4 gl_Position;\n"
-					"\tfloat gl_ClipDistance[6];\n"
+					"\tfloat gl_ClipDistance[];\n"
 					"};\n"
 					);
 #endif
@@ -3568,13 +3568,13 @@ public:
 				ralloc_asprintf_append(buffer, "in gl_PerVertex\n"
 					"{\n"
 					"\tvec4 gl_Position;\n"
-					"\tfloat gl_ClipDistance[6];\n"
+					"\tfloat gl_ClipDistance[];\n"
 					"} gl_in[gl_MaxPatchVertices];\n"
 					);
 				ralloc_asprintf_append(buffer, "out gl_PerVertex\n"
 					"{\n"
 					"\tvec4 gl_Position;\n"
-					"\tfloat gl_ClipDistance[6];\n"
+					"\tfloat gl_ClipDistance[];\n"
 					"} gl_out[];\n"
 					);
 				break;
@@ -3582,13 +3582,13 @@ public:
 				ralloc_asprintf_append(buffer, "in gl_PerVertex\n"
 					"{\n"
 					"\tvec4 gl_Position;\n"
-					"\tfloat gl_ClipDistance[6];\n"
+					"\tfloat gl_ClipDistance[];\n"
 					"} gl_in[gl_MaxPatchVertices];\n"
 					);
 				ralloc_asprintf_append(buffer, "out gl_PerVertex\n"
 					"{\n"
 					"\tvec4 gl_Position;\n"
-					"\tfloat gl_ClipDistance[6];\n"
+					"\tfloat gl_ClipDistance[];\n"
 					"};\n"
 					);
 				break;
