Fixed GLSL generation of SV_ClipDistance semantic in Vulkan backend.

diff --git a/Engine/Shaders/Public/Platform/Vulkan/VulkanCommon.ush b/Engine/Shaders/Public/Platform/Vulkan/VulkanCommon.ush
--- a/Engine/Shaders/Public/Platform/Vulkan/VulkanCommon.ush
+++ b/Engine/Shaders/Public/Platform/Vulkan/VulkanCommon.ush
@@ -5,7 +5,7 @@
 =============================================================================*/
 
 // Update this GUID to improve shader recompilation for Vulkan only shaders
-// GUID = {EDAF647B-ED20-440B-BE20-F5F9949CF71A}
+// GUID = {9230CFCD-FB48-4906-8EAF-B4ECF996BF34}
 // Glslang 1.1.82.1
 
 #pragma once
diff --git a/Engine/Source/Developer/VulkanShaderFormat/Private/VulkanBackend.cpp b/Engine/Source/Developer/VulkanShaderFormat/Private/VulkanBackend.cpp
--- a/Engine/Source/Developer/VulkanShaderFormat/Private/VulkanBackend.cpp
+++ b/Engine/Source/Developer/VulkanShaderFormat/Private/VulkanBackend.cpp
@@ -4541,7 +4541,7 @@ static ir_rvalue* GenShaderOutputSemantic(
 		}
 	}
 
-	if (Variable == NULL && Frequency == HSF_VertexShader)
+	if (Variable == NULL && (Frequency == HSF_VertexShader || Frequency == HSF_GeometryShader || Frequency == HSF_HullShader || Frequency == HSF_DomainShader))
 	{
 		const int PrefixLength = 15;
 		// Match SV_ClipDistance or SV_ClipDistanceN
