Fixed translation of "GetDimensions" intrinsic in VulkanBackend when texture is not used for sampling.

#rnx

diff --git a/Engine/Shaders/Private/Common.ush b/Engine/Shaders/Private/Common.ush
--- a/Engine/Shaders/Private/Common.ush
+++ b/Engine/Shaders/Private/Common.ush
@@ -72,10 +72,6 @@ const static float MaxHalfFloat = 65504.0f;
 #define USE_RAYTRACED_TEXTURE_RAYCONE_LOD (RAYHITGROUPSHADER)
 #endif
 
-#if COMPILER_HLSLCC
-	// Workaround until UE-69158 gets fixed
-	#define ComputeRayConeLod(Tex)	0
-#else
 static float GlobalRayCone_TexArea;
 float ComputeRayConeLod(Texture2D Tex)
 {
@@ -88,7 +84,6 @@ float ComputeRayConeLod(Texture2D Tex)
     return 0;
 #endif
 }
-#endif
 
 float  ClampToHalfFloatRange(float  X) { return clamp(X, float(0), MaxHalfFloat); }
 float2 ClampToHalfFloatRange(float2 X) { return clamp(X, float(0).xx, MaxHalfFloat.xx); }
diff --git a/Engine/Shaders/Public/Platform/Vulkan/VulkanCommon.ush b/Engine/Shaders/Public/Platform/Vulkan/VulkanCommon.ush
--- a/Engine/Shaders/Public/Platform/Vulkan/VulkanCommon.ush
+++ b/Engine/Shaders/Public/Platform/Vulkan/VulkanCommon.ush
@@ -5,7 +5,7 @@
 =============================================================================*/
 
 // Update this GUID to improve shader recompilation for Vulkan only shaders
-// GUID = 8F54F78F-6560-4404-B58B-D11D22128DFC
+// GUID = 86A78414-CF0D-48DB-AC9C-8A8301C857A5
 // Glslang 1.1.82.1
 
 #pragma once
diff --git a/Engine/Source/Developer/VulkanShaderFormat/Private/VulkanBackend.cpp b/Engine/Source/Developer/VulkanShaderFormat/Private/VulkanBackend.cpp
--- a/Engine/Source/Developer/VulkanShaderFormat/Private/VulkanBackend.cpp
+++ b/Engine/Source/Developer/VulkanShaderFormat/Private/VulkanBackend.cpp
@@ -429,7 +429,7 @@ struct FSamplerMappingGatherData
 {
 	struct FEntry
 	{
-		bool bUsingLoad = false;
+		bool bUsingLoadOrDim = false; // either "Load" or "GetDimensions" intrinsics are used
 		TStringSet SamplerStates;
 	};
 	std::map<std::string, FEntry> Entries;
@@ -451,7 +451,7 @@ struct FSamplerMapping
 		// First find all samplers using T.Load()
 		for (auto EntryPair : GatherData.Entries)
 		{
-			if (EntryPair.second.bUsingLoad)
+			if (EntryPair.second.bUsingLoadOrDim)
 			{
 				CombinedSamplers.insert(EntryPair.first);
 			}
@@ -3831,23 +3831,16 @@ struct FGenerateSamplerToTextureMapVisitor : public ir_hierarchical_visitor
 			}
 			else
 			{
-				if (IR->op == ir_txf)
+				if (IR->op == ir_txf || IR->op == ir_txs || IR->op == ir_txm)
 				{
-					GatherData.Entries[Sampler->name].bUsingLoad = true;
+					GatherData.Entries[Sampler->name].bUsingLoadOrDim = true;
 				}
 				else
 				{
-					if (IR->op == ir_txs || IR->op == ir_txm)
-					{
-						// Will be patched later
-					}
-					else
-					{
 #if UE_BUILD_DEBUG
-						// Internal error!!!
-						ensure(0);
+					// Internal error!!!
+					ensure(0);
 #endif
-					}
 					GatherData.Entries[Sampler->name].SamplerStates.insert("");
 				}
 			}
