Fixed erroneous number of texture samplers in material editor stats window generated by Vulkan shader compiler.

#rnx

diff --git a/Engine/Shaders/Public/Platform/Vulkan/VulkanCommon.ush b/Engine/Shaders/Public/Platform/Vulkan/VulkanCommon.ush
--- a/Engine/Shaders/Public/Platform/Vulkan/VulkanCommon.ush
+++ b/Engine/Shaders/Public/Platform/Vulkan/VulkanCommon.ush
@@ -5,7 +5,7 @@
 =============================================================================*/
 
 // Update this GUID to improve shader recompilation for Vulkan only shaders
-// GUID = 879F6C92-E759-47F4-A195-9F5FD9945012
+// GUID = 3720886C-A634-4CD5-99E5-728284036F3C
 // Glslang 1.1.82.1
 
 #pragma once
diff --git a/Engine/Source/Developer/VulkanShaderFormat/Private/VulkanShaderCompiler.cpp b/Engine/Source/Developer/VulkanShaderFormat/Private/VulkanShaderCompiler.cpp
--- a/Engine/Source/Developer/VulkanShaderFormat/Private/VulkanShaderCompiler.cpp
+++ b/Engine/Source/Developer/VulkanShaderFormat/Private/VulkanShaderCompiler.cpp
@@ -1360,6 +1360,9 @@ static void BuildShaderOutput(
 			EShaderParameterType::Sampler
 		);
 		NEWEntryTypes.Add(Name, FVulkanShaderHeader::Global);
+
+		// Count only samplers states, not textures
+		OLDHeader.SerializedBindings.NumSamplers++;
 	}
 
 	for (auto& Sampler : CCHeader.Samplers)
@@ -1372,14 +1375,9 @@ static void BuildShaderOutput(
 			VulkanBindingIndex,
 			Sampler.Count,
 			EShaderParameterType::SRV
-			);
+		);
 		NEWEntryTypes.Add(Sampler.Name, FVulkanShaderHeader::Global);
 
-		OLDHeader.SerializedBindings.NumSamplers = FMath::Max<uint8>(
-			OLDHeader.SerializedBindings.NumSamplers,
-			Sampler.Offset + Sampler.Count
-			);
-
 		for (auto& SamplerState : Sampler.SamplerStates)
 		{
 			if (!SharedSamplerStates.Contains(SamplerState))
@@ -1394,6 +1392,9 @@ static void BuildShaderOutput(
 					EShaderParameterType::Sampler
 				);
 				NEWEntryTypes.Add(SamplerState, FVulkanShaderHeader::Global);
+
+				// Count compiled texture-samplers as output samplers
+				OLDHeader.SerializedBindings.NumSamplers += Sampler.Count;
 			}
 		}
 	}
