Fix Vulkan Shader Compiler bug: no longer strips multiview shader permutation, by taking VIEW_ID into account.

diff --git a/Engine/Source/Developer/VulkanShaderFormat/Private/VulkanShaderCompiler.cpp b/Engine/Source/Developer/VulkanShaderFormat/Private/VulkanShaderCompiler.cpp
--- a/Engine/Source/Developer/VulkanShaderFormat/Private/VulkanShaderCompiler.cpp
+++ b/Engine/Source/Developer/VulkanShaderFormat/Private/VulkanShaderCompiler.cpp
@@ -1944,7 +1944,7 @@ void DoCompileVulkanShader(const FShaderCompilerInput& Input, FShaderCompilerOut
 		}
 	}
 
-	if (!PreprocessedShaderSource.Contains(TEXT("SV_ViewID")))
+	if (!(PreprocessedShaderSource.Contains(TEXT("SV_ViewID")) || PreprocessedShaderSource.Contains(TEXT("VIEW_ID"))))
 	{
 		// Disable instanced stereo if not requested by the shader
 		StripInstancedStereo(PreprocessedShaderSource);
