Define ES3_1_PROFILE for vulkan on mobile instead of ES2_PROFILE

diff --git a/Engine/Source/Developer/VulkanShaderFormat/Private/VulkanShaderCompiler.cpp b/Engine/Source/Developer/VulkanShaderFormat/Private/VulkanShaderCompiler.cpp
--- a/Engine/Source/Developer/VulkanShaderFormat/Private/VulkanShaderCompiler.cpp
+++ b/Engine/Source/Developer/VulkanShaderFormat/Private/VulkanShaderCompiler.cpp
@@ -1900,7 +1900,7 @@ void DoCompileVulkanShader(const FShaderCompilerInput& Input, FShaderCompilerOut
 	{
 		HlslCompilerTarget = HCT_FeatureLevelES3_1Ext;
 		HlslCompilerTargetES = HCT_FeatureLevelES3_1Ext;
-		AdditionalDefines.SetDefine(TEXT("ES2_PROFILE"), 1);
+		AdditionalDefines.SetDefine(TEXT("ES3_1_PROFILE"), 1);
 		AdditionalDefines.SetDefine(TEXT("VULKAN_PROFILE"), 1);
 	}
 	else if (bIsSM5)
diff --git a/Engine/Source/Developer/VulkanShaderFormat/Private/VulkanShaderFormat.cpp b/Engine/Source/Developer/VulkanShaderFormat/Private/VulkanShaderFormat.cpp
--- a/Engine/Source/Developer/VulkanShaderFormat/Private/VulkanShaderFormat.cpp
+++ b/Engine/Source/Developer/VulkanShaderFormat/Private/VulkanShaderFormat.cpp
@@ -24,8 +24,8 @@ class FShaderFormatVulkan : public IShaderFormat
 {
 	enum 
 	{
-		UE_SHADER_VULKAN_ES3_1_VER	= 27,
-		UE_SHADER_VULKAN_SM5_VER 	= 27,
+		UE_SHADER_VULKAN_ES3_1_VER	= 28,
+		UE_SHADER_VULKAN_SM5_VER 	= 28,
 	};
 
 	int32 InternalGetVersion(FName Format) const
