F - Write out info file containing fixed Vulkan shader SPIR-V entry point
#rnx

diff --git a/Engine/Source/Developer/VulkanShaderFormat/Private/VulkanShaderCompiler.cpp b/Engine/Source/Developer/VulkanShaderFormat/Private/VulkanShaderCompiler.cpp
--- a/Engine/Source/Developer/VulkanShaderFormat/Private/VulkanShaderCompiler.cpp
+++ b/Engine/Source/Developer/VulkanShaderFormat/Private/VulkanShaderCompiler.cpp
@@ -1757,6 +1757,20 @@ static bool CompileUsingInternal(FCompilerInfo& CompilerInfo, FVulkanBindingTabl
 		BuildShaderOutput(Output, CompilerInfo.Input,
 			GlslSource.GetData(), GlslSource.Num(),
 			BindingTable, NumLines, Spirv, DebugName, bHasRealUBs);
+
+		if (CompilerInfo.bDebugDump)
+		{
+			FString InfoFile = CompilerInfo.Input.DumpDebugInfoPath / TEXT("Info.txt");
+			FArchive* FileWriter = IFileManager::Get().CreateFileWriter(*InfoFile);
+			if (FileWriter)
+			{
+				FString OutputString = FString::Printf(TEXT("main_%0.8x_%0.8x\n"), Spirv.Data.Num() * sizeof(uint32), Spirv.CRC);
+				auto AnsiOutputString = StringCast<ANSICHAR>(*OutputString);
+				FileWriter->Serialize((ANSICHAR*)AnsiOutputString.Get(), AnsiOutputString.Length());
+				FileWriter->Close();
+			}
+			delete FileWriter;
+		}
 		return true;
 	}
 	else
