Prevent Vulkan early fragment tests being enabled when discard instruction is used.

diff --git a/Engine/Source/Developer/VulkanShaderFormat/Private/VulkanBackend.cpp b/Engine/Source/Developer/VulkanShaderFormat/Private/VulkanBackend.cpp
--- a/Engine/Source/Developer/VulkanShaderFormat/Private/VulkanBackend.cpp
+++ b/Engine/Source/Developer/VulkanShaderFormat/Private/VulkanBackend.cpp
@@ -772,6 +772,9 @@ class FGenerateVulkanVisitor : public ir_visitor
 	/** Found dFdx or dFdy */
 	bool bUsesDXDY;
 
+	// True if the discard instruction was encountered.
+	bool bUsesDiscard;
+
 	/** Found image atomic functions (e.g. imageAtomicAdd) */
 	bool bUsesImageWriteAtomic;
 
@@ -2325,6 +2328,7 @@ class FGenerateVulkanVisitor : public ir_visitor
 			ralloc_asprintf_append(buffer, ") ");
 		}
 		ralloc_asprintf_append(buffer, "discard");
+		bUsesDiscard = true;
 	}
 
 	bool try_conditional_move(ir_if *expr)
@@ -3272,7 +3276,7 @@ class FGenerateVulkanVisitor : public ir_visitor
 	*/
 	void print_layout(_mesa_glsl_parse_state *state)
 	{
-		if (early_depth_stencil)
+		if (early_depth_stencil && this->bUsesDiscard == false)
 		{
 			ralloc_asprintf_append(buffer, "layout(early_fragment_tests) in;\n");
 		}
@@ -3429,6 +3433,7 @@ public:
 		, loop_count(0)
 		, bUsesES2TextureLODExtension(false)
 		, bUsesDXDY(false)
+		, bUsesDiscard(false)
 		, bUsesImageWriteAtomic(false)
 	{
 		printable_names = hash_table_ctor(32, hash_table_pointer_hash, hash_table_pointer_compare);
diff --git a/Engine/Source/Developer/VulkanShaderFormat/Private/VulkanShaderFormat.cpp b/Engine/Source/Developer/VulkanShaderFormat/Private/VulkanShaderFormat.cpp
--- a/Engine/Source/Developer/VulkanShaderFormat/Private/VulkanShaderFormat.cpp
+++ b/Engine/Source/Developer/VulkanShaderFormat/Private/VulkanShaderFormat.cpp
@@ -24,8 +24,8 @@ class FShaderFormatVulkan : public IShaderFormat
 {
 	enum 
 	{
-		UE_SHADER_VULKAN_ES3_1_VER	= 28,
-		UE_SHADER_VULKAN_SM5_VER 	= 28,
+		UE_SHADER_VULKAN_ES3_1_VER	= 29,
+		UE_SHADER_VULKAN_SM5_VER 	= 29,
 	};
 
 	int32 InternalGetVersion(FName Format) const
