F - Fix Vulkan codegen for RWStructuredBuffer<float4> not generating scalar access

diff --git a/Engine/Shaders/Public/Platform/Vulkan/VulkanCommon.ush b/Engine/Shaders/Public/Platform/Vulkan/VulkanCommon.ush
--- a/Engine/Shaders/Public/Platform/Vulkan/VulkanCommon.ush
+++ b/Engine/Shaders/Public/Platform/Vulkan/VulkanCommon.ush
@@ -5,7 +5,7 @@
 =============================================================================*/
 
 // Update this GUID to improve shader recompilation for Vulkan only shaders
-// GUID = 024CC127-CDD1-424E-9200-B3AB85F016B4
+// GUID = 58607BA4-C1FE-47DC-ADDF-15496DAC9772
 // Glslang 1.1.82.1
 
 #pragma once
diff --git a/Engine/Source/Developer/VulkanShaderFormat/Private/VulkanBackend.cpp b/Engine/Source/Developer/VulkanShaderFormat/Private/VulkanBackend.cpp
--- a/Engine/Source/Developer/VulkanShaderFormat/Private/VulkanBackend.cpp
+++ b/Engine/Source/Developer/VulkanShaderFormat/Private/VulkanBackend.cpp
@@ -1891,7 +1891,7 @@ class FGenerateVulkanVisitor : public ir_visitor
 		}
 	}
 
-	void print_image_op(ir_dereference_image *deref, ir_rvalue *src)
+	void print_image_op(ir_dereference_image *deref, ir_rvalue *src, const char* Mask = nullptr)
 	{
 		const char* swizzle[] =
 		{
@@ -1927,6 +1927,10 @@ class FGenerateVulkanVisitor : public ir_visitor
 					ralloc_asprintf_append(buffer, "[");
 					deref->image_index->accept(this);
 					ralloc_asprintf_append(buffer, "]");
+					if (Mask)
+					{
+						ralloc_asprintf_append(buffer, Mask);
+					}
 					ralloc_asprintf_append(buffer, " = ");
 					src->accept(this);
 				}
@@ -2041,31 +2045,30 @@ class FGenerateVulkanVisitor : public ir_visitor
 			ralloc_asprintf_append(buffer, ") { ");
 		}
 
-		if (assign->lhs->as_dereference_image() != NULL)
-		{
-			/** EHart - should the write mask be checked here? */
-			print_image_op(assign->lhs->as_dereference_image(), assign->rhs);
-		}
-		else
+		char mask[6] = "";
+		unsigned j = 1;
+		if (assign->lhs->type->is_scalar() == false ||
+			assign->write_mask != 0x1)
 		{
-			char mask[6];
-			unsigned j = 1;
-			if (assign->lhs->type->is_scalar() == false ||
-				assign->write_mask != 0x1)
+			for (unsigned i = 0; i < 4; i++)
 			{
-				for (unsigned i = 0; i < 4; i++)
+				if ((assign->write_mask & (1 << i)) != 0)
 				{
-					if ((assign->write_mask & (1 << i)) != 0)
-					{
-						mask[j] = "xyzw"[i];
-						j++;
-					}
+					mask[j] = "xyzw"[i];
+					j++;
 				}
 			}
-			mask[j] = '\0';
+		}
+		mask[j] = '\0';
 
-			mask[0] = (j == 1) ? '\0' : '.';
+		mask[0] = (j == 1) ? '\0' : '.';
 
+		if (assign->lhs->as_dereference_image() != NULL)
+		{
+			print_image_op(assign->lhs->as_dereference_image(), assign->rhs, mask);
+		}
+		else
+		{
 			assign->lhs->accept(this);
 			ralloc_asprintf_append(buffer, "%s = ", mask);
 			assign->rhs->accept(this);
