Fixed overlapping shader attributes in Vulkan when using matrix types.

The code was assuming that each attribute occupies a single vector, so we got compiler errors if a shader used a matrix attribute.

#rnx

diff --git a/Engine/Source/Developer/VulkanShaderFormat/Private/VulkanBackend.cpp b/Engine/Source/Developer/VulkanShaderFormat/Private/VulkanBackend.cpp
--- a/Engine/Source/Developer/VulkanShaderFormat/Private/VulkanBackend.cpp
+++ b/Engine/Source/Developer/VulkanShaderFormat/Private/VulkanBackend.cpp
@@ -4214,11 +4214,13 @@ static void ConfigureInOutVariableLayout(EHlslShaderFrequency Frequency,
 		Variable->semantic = ralloc_strdup(Variable, Semantic);
 		if (Mode == ir_var_in)
 		{
-			Variable->location = ParseState->next_in_location_slot++;
+			Variable->location = ParseState->next_in_location_slot;
+			ParseState->next_in_location_slot += Variable->type->matrix_columns;
 		}
 		else
 		{
-			Variable->location = ParseState->next_out_location_slot++;
+			Variable->location = ParseState->next_out_location_slot;
+			ParseState->next_out_location_slot += Variable->type->matrix_columns;
 		}
 	}
 }
