Fixed bug introduced in CL 12982657 which caused compile failures for some geometry shaders on Vulkan.

Arrays don't have matrix_columns set to 1.

#rnx

diff --git a/Engine/Source/Developer/VulkanShaderFormat/Private/VulkanBackend.cpp b/Engine/Source/Developer/VulkanShaderFormat/Private/VulkanBackend.cpp
--- a/Engine/Source/Developer/VulkanShaderFormat/Private/VulkanBackend.cpp
+++ b/Engine/Source/Developer/VulkanShaderFormat/Private/VulkanBackend.cpp
@@ -4212,15 +4212,16 @@ static void ConfigureInOutVariableLayout(EHlslShaderFrequency Frequency,
 	{
 		Variable->explicit_location = 1;
 		Variable->semantic = ralloc_strdup(Variable, Semantic);
+		unsigned int NumVectors = (Variable->type->matrix_columns > 1) ? Variable->type->matrix_columns : 1;
 		if (Mode == ir_var_in)
 		{
 			Variable->location = ParseState->next_in_location_slot;
-			ParseState->next_in_location_slot += Variable->type->matrix_columns;
+			ParseState->next_in_location_slot += NumVectors;
 		}
 		else
 		{
 			Variable->location = ParseState->next_out_location_slot;
-			ParseState->next_out_location_slot += Variable->type->matrix_columns;
+			ParseState->next_out_location_slot += NumVectors;
 		}
 	}
 }
