Fixed offset for geometry shader input slots of array types in HLSLcc/Vulkan backend.

#rnx

diff --git a/Engine/Source/Developer/VulkanShaderFormat/Private/VulkanBackend.cpp b/Engine/Source/Developer/VulkanShaderFormat/Private/VulkanBackend.cpp
--- a/Engine/Source/Developer/VulkanShaderFormat/Private/VulkanBackend.cpp
+++ b/Engine/Source/Developer/VulkanShaderFormat/Private/VulkanBackend.cpp
@@ -4177,6 +4177,32 @@ static FSystemValue* SystemValueTable[HSF_FrequencyCount] =
 
 #define CUSTOM_LAYER_INDEX_SEMANTIC "HLSLCC_LAYER_INDEX"
 
+// Returns the number of slots an in/out variable occupies. This is determined by the variable type, e.g. vec4 has offset 1, mat4 has offset 4.
+// For geometry shader input variables, the first array specifier is ignored as geometry shaders get their input from multiple vertex shader invocations.
+static unsigned GetInOutVariableSlotSize(EHlslShaderFrequency Frequency, ir_variable_mode Mode, const struct glsl_type* VarType)
+{
+	check(VarType != nullptr);
+	if (VarType->fields.array != nullptr)
+	{
+		if (Frequency == HSF_GeometryShader && Mode == ir_var_in)
+		{
+			// Ignore frequency after first array specifier, so use HSF_InvalidFrequency
+			return GetInOutVariableSlotSize(HSF_InvalidFrequency, Mode, VarType->fields.array);
+		}
+		else
+		{
+			// Multiply inner array type by array size, i.e. both 'vec2[4]' and 'float[4]' occupy 4 slots
+			// Ignore frequency after first array specifier, so use HSF_InvalidFrequency
+			return GetInOutVariableSlotSize(HSF_InvalidFrequency, Mode, VarType->fields.array) * VarType->length;
+		}
+	}
+	else
+	{
+		// Only use number of matrix columns to determine number of slots this variable occupies: 1 for scalars and vectors, 2/3/4 for matrices
+		return VarType->matrix_columns;
+	}
+}
+
 static void ConfigureInOutVariableLayout(EHlslShaderFrequency Frequency,
 	_mesa_glsl_parse_state* ParseState,
 	const char* Semantic,
@@ -4216,12 +4242,12 @@ static void ConfigureInOutVariableLayout(EHlslShaderFrequency Frequency,
 		if (Mode == ir_var_in)
 		{
 			Variable->location = ParseState->next_in_location_slot;
-			ParseState->next_in_location_slot += NumVectors;
+			ParseState->next_in_location_slot += GetInOutVariableSlotSize(Frequency, Mode, Variable->type);
 		}
 		else
 		{
 			Variable->location = ParseState->next_out_location_slot;
-			ParseState->next_out_location_slot += NumVectors;
+			ParseState->next_out_location_slot += GetInOutVariableSlotSize(Frequency, Mode, Variable->type);
 		}
 	}
 }
