Improve assertion about missing entry in SPIR-V module in Vulkan backend.

#rnx

diff --git a/Engine/Source/Developer/VulkanShaderFormat/Private/VulkanGlslangBridge.cpp b/Engine/Source/Developer/VulkanShaderFormat/Private/VulkanGlslangBridge.cpp
--- a/Engine/Source/Developer/VulkanShaderFormat/Private/VulkanGlslangBridge.cpp
+++ b/Engine/Source/Developer/VulkanShaderFormat/Private/VulkanGlslangBridge.cpp
@@ -310,7 +310,7 @@ static void ComputeMovableWordIndices(FSpirv& Spirv)
 				if (FoundTypeName && FoundTypeName->Len() > 0)
 				{
 					FSpirv::FEntry* FoundEntry = Spirv.GetEntry(*FoundTypeName);
-					check(FoundEntry);
+					checkf(FoundEntry, TEXT("cannot find variable entry in SPIR-V module: %s"), **FoundTypeName);
 					FDecorations& FoundDecorations = Decorations.FindChecked(VariableId);
 					FoundEntry->Binding = FoundDecorations.BindingIndex;
 					FoundEntry->WordBindingIndex = FoundDecorations.WordBindingIndex;
@@ -322,7 +322,7 @@ static void ComputeMovableWordIndices(FSpirv& Spirv)
 			{
 				// Standalone global var
 				FSpirv::FEntry* FoundEntry = Spirv.GetEntry(*FoundVariableName);
-				check(FoundEntry);
+				checkf(FoundEntry, TEXT("cannot find variable entry in SPIR-V module: %s"), **FoundVariableName);
 				FDecorations& FoundDecorations = Decorations.FindChecked(VariableId);
 				FoundEntry->Binding = FoundDecorations.BindingIndex;
 				FoundEntry->WordBindingIndex = FoundDecorations.WordBindingIndex;
