Fix SPIR-V reflection of texture SRVs in Vulkan/DXC backend. Merged from CL 13654014.

#rnx

diff --git a/Engine/Source/Developer/VulkanShaderFormat/Private/VulkanShaderCompiler.cpp b/Engine/Source/Developer/VulkanShaderFormat/Private/VulkanShaderCompiler.cpp
--- a/Engine/Source/Developer/VulkanShaderFormat/Private/VulkanShaderCompiler.cpp
+++ b/Engine/Source/Developer/VulkanShaderFormat/Private/VulkanShaderCompiler.cpp
@@ -2623,7 +2623,7 @@ static void BuildShaderOutputFromSPIRV(
 		const FString ResourceName(ANSI_TO_TCHAR(Binding->name));
 		if (Binding->usage_binding_count > 0)
 		{
-			SRVString += FString::Printf(TEXT("%s%s(%u:%u["), SRVString.Len() ? TEXT(",") : TEXT(""), *ResourceName, SRVBindings++, Binding->usage_binding_count);
+			SRVString += FString::Printf(TEXT("%s%s(%u:%u["), SRVString.Len() ? TEXT(",") : TEXT(""), *ResourceName, SRVBindings++, 1);
 
 			for (uint32 UsageIndex = 0; UsageIndex < Binding->usage_binding_count; ++UsageIndex)
 			{
