Fix for shader compile tasks when building for Vulkan when check is disabled

diff --git a/Engine/Source/Developer/VulkanShaderFormat/Private/VulkanGlslangBridge.cpp b/Engine/Source/Developer/VulkanShaderFormat/Private/VulkanGlslangBridge.cpp
--- a/Engine/Source/Developer/VulkanShaderFormat/Private/VulkanGlslangBridge.cpp
+++ b/Engine/Source/Developer/VulkanShaderFormat/Private/VulkanGlslangBridge.cpp
@@ -172,11 +172,13 @@ static void ComputeMovableWordIndices(FSpirv& Spirv)
 	const uint32* PtrStart = Spirv.Data.GetData();
 	const uint32* Ptr = PtrStart;
 	const uint32* PtrEnd = Spirv.Data.GetData() + Spirv.Data.Num();
-	check(*Ptr++ == spv::MagicNumber);
+	const uint32_t MagicNumberValue = *Ptr++;
+	check(MagicNumberValue == spv::MagicNumber);
 	uint32_t Version = *Ptr++;
 	uint32_t Generator = *Ptr++;
 	uint32_t Bound = *Ptr++;
-	check(*Ptr++ == 0);
+	const uint32_t ZeroCheckValue = *Ptr++;
+	check(ZeroCheckValue == 0);
 
 	auto ReadLiteralString = [](const uint32_t* Ptr)
 	{
