Remove unused function StringToFile in VulkanShaderCompiler.cpp

#rnx

diff --git a/Engine/Source/Developer/VulkanShaderFormat/Private/VulkanShaderCompiler.cpp b/Engine/Source/Developer/VulkanShaderFormat/Private/VulkanShaderCompiler.cpp
--- a/Engine/Source/Developer/VulkanShaderFormat/Private/VulkanShaderCompiler.cpp
+++ b/Engine/Source/Developer/VulkanShaderFormat/Private/VulkanShaderCompiler.cpp
@@ -1532,26 +1532,6 @@ static void BuildShaderOutput(
 	CullGlobalUniformBuffers(ShaderInput.Environment.UniformBufferMap, ShaderOutput.ParameterMap);
 }
 
-
-static bool StringToFile(const FString& Filepath, const char* str)
-{
-	int32 StrLength = str ? FCStringAnsi::Strlen(str) : 0;
-
-	if(StrLength == 0)
-	{
-		return false;
-	}
-
-	if (TUniquePtr<FArchive> FileWriter = TUniquePtr<FArchive>(IFileManager::Get().CreateFileWriter(*Filepath)))
-	{
-		// const cast...
-		FileWriter->Serialize((void*)str, StrLength+1);
-		FileWriter->Close();
-	}
-
-	return true;
-}
-
 static char* PatchGLSLVersionPosition(const char* InSourceGLSL)
 {
 	if(!InSourceGLSL)
