Move SPIR-V utility functions for reflection and binding information into SpirvCommon.h/.cpp source files.

#rnx

diff --git a/Engine/Source/Developer/Apple/MetalShaderFormat/Private/MetalDerivedData.cpp b/Engine/Source/Developer/Apple/MetalShaderFormat/Private/MetalDerivedData.cpp
--- a/Engine/Source/Developer/Apple/MetalShaderFormat/Private/MetalDerivedData.cpp
+++ b/Engine/Source/Developer/Apple/MetalShaderFormat/Private/MetalDerivedData.cpp
@@ -13,12 +13,7 @@
 #include "HlslccHeaderWriter.h"
 
 #include "MetalShaderFormat.h"
-
-#if PLATFORM_MAC || PLATFORM_WINDOWS
-THIRD_PARTY_INCLUDES_START
-#include "spirv_reflect.h"
-THIRD_PARTY_INCLUDES_END
-#endif
+#include "SpirvReflectCommon.h"
 
 extern void BuildMetalShaderOutput(
 	FShaderCompilerOutput& ShaderOutput,
@@ -454,10 +449,8 @@ bool FMetalShaderOutputCooker::Build(TArray<uint8>& OutData)
 			
 			SpvReflectResult SPVRResult = SPV_REFLECT_RESULT_NOT_READY;
 			uint32 Count = 0;
+			FSpirvReflectBindings ReflectionBindings;
 			TArray<SpvReflectDescriptorBinding*> Bindings;
-			TSet<SpvReflectDescriptorBinding*> Counters;
-			TArray<SpvReflectInterfaceVariable*> InputVars;
-			TArray<SpvReflectInterfaceVariable*> OutputVars;
 			TArray<SpvReflectBlockVariable*> ConstantBindings;
 			TArray<SpvReflectExecutionMode*> ExecutionModes;
 			
@@ -572,19 +565,11 @@ bool FMetalShaderOutputCooker::Build(TArray<uint8>& OutData)
 			{
 				TArray<SpvReflectDescriptorBinding*> ResourceBindings;
 				TArray<SpvReflectDescriptorBinding*> ArgumentBindings;
-				TArray<SpvReflectDescriptorBinding*> UniformBindings;
-				TArray<SpvReflectDescriptorBinding*> SamplerBindings;
-				TArray<SpvReflectDescriptorBinding*> TextureSRVBindings;
-				TArray<SpvReflectDescriptorBinding*> TextureUAVBindings;
-				TArray<SpvReflectDescriptorBinding*> TBufferSRVBindings;
-				TArray<SpvReflectDescriptorBinding*> TBufferUAVBindings;
-				TArray<SpvReflectDescriptorBinding*> SBufferSRVBindings;
-				TArray<SpvReflectDescriptorBinding*> SBufferUAVBindings;
 				TSet<FString> UsedSets;
 				
 				// Extract all the bindings first so that we process them in order - this lets us assign UAVs before other resources
 				// Which is necessary to match the D3D binding scheme.
-				for (auto const& Binding : Bindings)
+				for (SpvReflectDescriptorBinding* Binding : Bindings)
 				{
 					if (Binding->resource_type != SPV_REFLECT_RESOURCE_FLAG_CBV && ResourceTable.Contains(UTF8_TO_TCHAR(Binding->name)))
 					{
@@ -595,110 +580,16 @@ bool FMetalShaderOutputCooker::Build(TArray<uint8>& OutData)
 						
 						continue;
 					}
-					
-					switch(Binding->resource_type)
+
+					// Add descriptor binding to argument bindings if it's a constant buffer with a name from 'TableNames'. Otherwise, add to common binding container.
+					if (Binding->resource_type == SPV_REFLECT_RESOURCE_FLAG_CBV && Binding->accessed && TableNames.Contains(UTF8_TO_TCHAR(Binding->name)))
 					{
-						case SPV_REFLECT_RESOURCE_FLAG_CBV:
-						{
-							check(Binding->descriptor_type == SPV_REFLECT_DESCRIPTOR_TYPE_UNIFORM_BUFFER);
-							if (Binding->accessed)
-							{
-								if (TableNames.Contains(UTF8_TO_TCHAR(Binding->name)))
-								{
-									ArgumentBindings.Add(Binding);
-								}
-								else
-								{
-									UniformBindings.Add(Binding);
-								}
-							}
-							break;
-						}
-						case SPV_REFLECT_RESOURCE_FLAG_SAMPLER:
-						{
-							check(Binding->descriptor_type == SPV_REFLECT_DESCRIPTOR_TYPE_SAMPLER);
-							if (Binding->accessed)
-							{
-								SamplerBindings.Add(Binding);
-							}
-							break;
-						}
-						case SPV_REFLECT_RESOURCE_FLAG_SRV:
-						{
-							switch(Binding->descriptor_type)
-							{
-								case SPV_REFLECT_DESCRIPTOR_TYPE_SAMPLED_IMAGE:
-								{
-									if (Binding->accessed)
-									{
-										TextureSRVBindings.Add(Binding);
-									}
-									break;
-								}
-								case SPV_REFLECT_DESCRIPTOR_TYPE_UNIFORM_TEXEL_BUFFER:
-								{
-									if (Binding->accessed)
-									{
-										TBufferSRVBindings.Add(Binding);
-									}
-									break;
-								}
-								case SPV_REFLECT_DESCRIPTOR_TYPE_STORAGE_BUFFER:
-								{
-									if (Binding->accessed)
-									{
-										SBufferSRVBindings.Add(Binding);
-									}
-									break;
-								}
-								default:
-								{
-									// check(false);
-									break;
-								}
-							}
-							break;
-						}
-						case SPV_REFLECT_RESOURCE_FLAG_UAV:
-						{
-							if (Binding->uav_counter_binding)
-							{
-								Counters.Add(Binding->uav_counter_binding);
-							}
-							
-							switch(Binding->descriptor_type)
-							{
-								case SPV_REFLECT_DESCRIPTOR_TYPE_STORAGE_IMAGE:
-								{
-									TextureUAVBindings.Add(Binding);
-									break;
-								}
-								case SPV_REFLECT_DESCRIPTOR_TYPE_STORAGE_TEXEL_BUFFER:
-								{
-									TBufferUAVBindings.Add(Binding);
-									break;
-								}
-								case SPV_REFLECT_DESCRIPTOR_TYPE_STORAGE_BUFFER:
-								{
-									if (!Counters.Contains(Binding) || Binding->accessed)
-									{
-										SBufferUAVBindings.Add(Binding);
-									}
-									break;
-								}
-								default:
-								{
-									// check(false);
-									break;
-								}
-							}
-							break;
-						}
-						default:
-						{
-							// check(false);
-							break;
-						}
+						check(Binding->descriptor_type == SPV_REFLECT_DESCRIPTOR_TYPE_UNIFORM_BUFFER);
+						ArgumentBindings.Add(Binding);
+					}
+					else
+					{
+						ReflectionBindings.AddDescriptorBinding(Binding);
 					}
 				}
 				
@@ -750,14 +641,14 @@ bool FMetalShaderOutputCooker::Build(TArray<uint8>& OutData)
 					}
 					else
 					{
-						UniformBindings.Add(ArgumentBindings[i]);
+						ReflectionBindings.UniformBuffers.Add(ArgumentBindings[i]);
 						ArgumentBindings.RemoveAt(i);
 					}
 				}
 				
-				uint32 GlobalSetId = 32;
+				const uint32 GlobalSetId = 32;
 				
-				for (auto const& Binding : TBufferUAVBindings)
+				for (auto const& Binding : ReflectionBindings.TBufferUAVs)
 				{
 					check(UAVIndices);
 					uint32 Index = FPlatformMath::CountTrailingZeros(UAVIndices);
@@ -776,7 +667,7 @@ bool FMetalShaderOutputCooker::Build(TArray<uint8>& OutData)
 					check(SPVRResult == SPV_REFLECT_RESULT_SUCCESS);
 				}
 				
-				for (auto const& Binding : SBufferUAVBindings)
+				for (auto const& Binding : ReflectionBindings.SBufferUAVs)
 				{
 					check(UAVIndices);
 					uint32 Index = FPlatformMath::CountTrailingZeros(UAVIndices);
@@ -794,7 +685,7 @@ bool FMetalShaderOutputCooker::Build(TArray<uint8>& OutData)
 					check(SPVRResult == SPV_REFLECT_RESULT_SUCCESS);
 				}
 				
-				for (auto const& Binding : TextureUAVBindings)
+				for (auto const& Binding : ReflectionBindings.TextureUAVs)
 				{
 					check(UAVIndices);
 					uint32 Index = FPlatformMath::CountTrailingZeros(UAVIndices);
@@ -925,7 +816,7 @@ bool FMetalShaderOutputCooker::Build(TArray<uint8>& OutData)
 					CCHeaderWriter.WriteUniformBlock(*Name, SetIndex);
 				}
 				
-				for (auto const& Binding : TBufferSRVBindings)
+				for (auto const& Binding : ReflectionBindings.TBufferSRVs)
 				{
 					check(TextureIndices);
 					uint32 Index = FPlatformMath::CountTrailingZeros(TextureIndices);
@@ -942,7 +833,7 @@ bool FMetalShaderOutputCooker::Build(TArray<uint8>& OutData)
 					check(SPVRResult == SPV_REFLECT_RESULT_SUCCESS);
 				}
 				
-				for (auto const& Binding : SBufferSRVBindings)
+				for (auto const& Binding : ReflectionBindings.SBufferSRVs)
 				{
 					check(BufferIndices);
 					uint32 Index = FPlatformMath::CountTrailingZeros64(BufferIndices);
@@ -957,7 +848,7 @@ bool FMetalShaderOutputCooker::Build(TArray<uint8>& OutData)
 					check(SPVRResult == SPV_REFLECT_RESULT_SUCCESS);
 				}
 				
-				for (auto const& Binding : UniformBindings)
+				for (auto const& Binding : ReflectionBindings.UniformBuffers)
 				{
 					check(BufferIndices);
 					uint32 Index = FPlatformMath::CountTrailingZeros64(BufferIndices);
@@ -1008,7 +899,7 @@ bool FMetalShaderOutputCooker::Build(TArray<uint8>& OutData)
 					check(SPVRResult == SPV_REFLECT_RESULT_SUCCESS);
 				}
 				
-				for (auto const& Binding : TextureSRVBindings)
+				for (auto const& Binding : ReflectionBindings.TextureSRVs)
 				{
 					check(TextureIndices);
 					uint32 Index = FPlatformMath::CountTrailingZeros64(TextureIndices);
@@ -1020,7 +911,7 @@ bool FMetalShaderOutputCooker::Build(TArray<uint8>& OutData)
 					check(SPVRResult == SPV_REFLECT_RESULT_SUCCESS);
 				}
 				
-				for (auto const& Binding : SamplerBindings)
+				for (auto const& Binding : ReflectionBindings.Samplers)
 				{
 					check(SamplerIndices);
 					uint32 Index = FPlatformMath::CountTrailingZeros64(SamplerIndices);
@@ -1035,131 +926,115 @@ bool FMetalShaderOutputCooker::Build(TArray<uint8>& OutData)
 			
 			if (Frequency == HSF_PixelShader)
 			{
-				Count = 0;
-				SPVRResult = Reflection.EnumerateOutputVariables(&Count, nullptr);
-				check(SPVRResult == SPV_REFLECT_RESULT_SUCCESS);
-				OutputVars.SetNum(Count);
-				SPVRResult = Reflection.EnumerateOutputVariables(&Count, OutputVars.GetData());
-				check(SPVRResult == SPV_REFLECT_RESULT_SUCCESS);
-				if (Count > 0)
+				ReflectionBindings.GatherOutputAttributes(Reflection);
+				for (auto const& Var : ReflectionBindings.OutputAttributes)
 				{
-					for (auto const& Var : OutputVars)
+					if (Var->storage_class == SpvStorageClassOutput && Var->built_in == -1 && strstr(Var->name, "SV_Target"))
 					{
-						if (Var->storage_class == SpvStorageClassOutput && Var->built_in == -1 && strstr(Var->name, "SV_Target"))
-						{
-							FString TypeQualifier;
+						FString TypeQualifier;
 							
-							auto const type = *Var->type_description;
-							uint32_t masked_type = type.type_flags & 0xF;
+						auto const type = *Var->type_description;
+						uint32_t masked_type = type.type_flags & 0xF;
 							
-							switch (masked_type) {
-								default: checkf(false, TEXT("unsupported component type %d"), masked_type); break;
-								case SPV_REFLECT_TYPE_FLAG_BOOL  : TypeQualifier = TEXT("b"); break;
-								case SPV_REFLECT_TYPE_FLAG_INT   : TypeQualifier = (type.traits.numeric.scalar.signedness ? TEXT("i") : TEXT("u")); break;
-								case SPV_REFLECT_TYPE_FLAG_FLOAT : TypeQualifier = (type.traits.numeric.scalar.width == 32 ? TEXT("f") : TEXT("h")); break;
-							}
-							
-							if (type.type_flags & SPV_REFLECT_TYPE_FLAG_MATRIX)
-							{
-								TypeQualifier += FString::Printf(TEXT("%d%d"), type.traits.numeric.matrix.row_count, type.traits.numeric.matrix.column_count);
-							}
-							else if (type.type_flags & SPV_REFLECT_TYPE_FLAG_VECTOR)
-							{
-								TypeQualifier += FString::Printf(TEXT("%d"), type.traits.numeric.vector.component_count);
-							}
-							else
-							{
-								TypeQualifier += TEXT("1");
-							}
+						switch (masked_type) {
+							default: checkf(false, TEXT("unsupported component type %d"), masked_type); break;
+							case SPV_REFLECT_TYPE_FLAG_BOOL  : TypeQualifier = TEXT("b"); break;
+							case SPV_REFLECT_TYPE_FLAG_INT   : TypeQualifier = (type.traits.numeric.scalar.signedness ? TEXT("i") : TEXT("u")); break;
+							case SPV_REFLECT_TYPE_FLAG_FLOAT : TypeQualifier = (type.traits.numeric.scalar.width == 32 ? TEXT("f") : TEXT("h")); break;
+						}
 							
-							CCHeaderWriter.WriteOutputAttribute(TEXT("SV_Target"), *TypeQualifier, Var->location, /*bLocationPrefix:*/ false, /*bLocationSuffix:*/ true);
+						if (type.type_flags & SPV_REFLECT_TYPE_FLAG_MATRIX)
+						{
+							TypeQualifier += FString::Printf(TEXT("%d%d"), type.traits.numeric.matrix.row_count, type.traits.numeric.matrix.column_count);
+						}
+						else if (type.type_flags & SPV_REFLECT_TYPE_FLAG_VECTOR)
+						{
+							TypeQualifier += FString::Printf(TEXT("%d"), type.traits.numeric.vector.component_count);
 						}
+						else
+						{
+							TypeQualifier += TEXT("1");
+						}
+							
+						CCHeaderWriter.WriteOutputAttribute(TEXT("SV_Target"), *TypeQualifier, Var->location, /*bLocationPrefix:*/ false, /*bLocationSuffix:*/ true);
 					}
 				}
 			}
 			
 			if (Frequency == HSF_VertexShader)
 			{
-				Count = 0;
-				SPVRResult = Reflection.EnumerateInputVariables(&Count, nullptr);
-				check(SPVRResult == SPV_REFLECT_RESULT_SUCCESS);
-				InputVars.SetNum(Count);
-				SPVRResult = Reflection.EnumerateInputVariables(&Count, InputVars.GetData());
-				check(SPVRResult == SPV_REFLECT_RESULT_SUCCESS);
-				if (Count > 0)
+				uint32 AssignedInputs = 0;
+
+				ReflectionBindings.GatherInputAttributes(Reflection);
+				for (auto const& Var : ReflectionBindings.InputAttributes)
 				{
-					uint32 AssignedInputs = 0;
-					
-					for (auto const& Var : InputVars)
+					if (Var->storage_class == SpvStorageClassInput && Var->built_in == -1)
 					{
-						if (Var->storage_class == SpvStorageClassInput && Var->built_in == -1)
+						unsigned Location = Var->location;
+						unsigned SemanticIndex = Location;
+						check(Var->semantic);
+						unsigned i = (unsigned)strlen(Var->semantic);
+						check(i);
+						while (isdigit((unsigned char)(Var->semantic[i-1])))
 						{
-							unsigned Location = Var->location;
-							unsigned SemanticIndex = Location;
-							check(Var->semantic);
-							unsigned i = (unsigned)strlen(Var->semantic);
-							check(i);
-							while (isdigit((unsigned char)(Var->semantic[i-1])))
-							{
-								i--;
-							}
-							if (i < strlen(Var->semantic))
+							i--;
+						}
+						if (i < strlen(Var->semantic))
+						{
+							SemanticIndex = (unsigned)atoi(Var->semantic + i);
+							if (Location != SemanticIndex)
 							{
-								SemanticIndex = (unsigned)atoi(Var->semantic + i);
-								if (Location != SemanticIndex)
-								{
-									Location = SemanticIndex;
-								}
+								Location = SemanticIndex;
 							}
+						}
 							
-							while ((1 << Location) & AssignedInputs)
-							{
-								Location++;
-							}
+						while ((1 << Location) & AssignedInputs)
+						{
+							Location++;
+						}
 							
-							if (Location != Var->location)
-							{
-								SPVRResult = Reflection.ChangeInputVariableLocation(Var, Location);
-								check(SPVRResult == SPV_REFLECT_RESULT_SUCCESS);
-							}
+						if (Location != Var->location)
+						{
+							SPVRResult = Reflection.ChangeInputVariableLocation(Var, Location);
+							check(SPVRResult == SPV_REFLECT_RESULT_SUCCESS);
+						}
 							
-							uint32 ArrayCount = 1;
-							for (uint32 Dim = 0; Dim < Var->array.dims_count; Dim++)
-							{
-								ArrayCount *= Var->array.dims[Dim];
-							}
+						uint32 ArrayCount = 1;
+						for (uint32 Dim = 0; Dim < Var->array.dims_count; Dim++)
+						{
+							ArrayCount *= Var->array.dims[Dim];
+						}
 							
-							FString TypeQualifier;
+						FString TypeQualifier;
 
-							auto const type = *Var->type_description;
-							uint32_t masked_type = type.type_flags & 0xF;
+						auto const type = *Var->type_description;
+						uint32_t masked_type = type.type_flags & 0xF;
 							
-							switch (masked_type) {
-								default: checkf(false, TEXT("unsupported component type %d"), masked_type); break;
-								case SPV_REFLECT_TYPE_FLAG_BOOL  : TypeQualifier = TEXT("b"); break;
-								case SPV_REFLECT_TYPE_FLAG_INT   : TypeQualifier = (type.traits.numeric.scalar.signedness ? TEXT("i") : TEXT("u")); break;
-								case SPV_REFLECT_TYPE_FLAG_FLOAT : TypeQualifier = (type.traits.numeric.scalar.width == 32 ? TEXT("f") : TEXT("h")); break;
-							}
+						switch (masked_type) {
+							default: checkf(false, TEXT("unsupported component type %d"), masked_type); break;
+							case SPV_REFLECT_TYPE_FLAG_BOOL  : TypeQualifier = TEXT("b"); break;
+							case SPV_REFLECT_TYPE_FLAG_INT   : TypeQualifier = (type.traits.numeric.scalar.signedness ? TEXT("i") : TEXT("u")); break;
+							case SPV_REFLECT_TYPE_FLAG_FLOAT : TypeQualifier = (type.traits.numeric.scalar.width == 32 ? TEXT("f") : TEXT("h")); break;
+						}
 							
-							if (type.type_flags & SPV_REFLECT_TYPE_FLAG_MATRIX)
-							{
-								TypeQualifier += FString::Printf(TEXT("%d%d"), type.traits.numeric.matrix.row_count, type.traits.numeric.matrix.column_count);
-							}
-							else if (type.type_flags & SPV_REFLECT_TYPE_FLAG_VECTOR)
-							{
-								TypeQualifier += FString::Printf(TEXT("%d"), type.traits.numeric.vector.component_count);
-							}
-							else
-							{
-								TypeQualifier += TEXT("1");
-							}
+						if (type.type_flags & SPV_REFLECT_TYPE_FLAG_MATRIX)
+						{
+							TypeQualifier += FString::Printf(TEXT("%d%d"), type.traits.numeric.matrix.row_count, type.traits.numeric.matrix.column_count);
+						}
+						else if (type.type_flags & SPV_REFLECT_TYPE_FLAG_VECTOR)
+						{
+							TypeQualifier += FString::Printf(TEXT("%d"), type.traits.numeric.vector.component_count);
+						}
+						else
+						{
+							TypeQualifier += TEXT("1");
+						}
 							
-							for (uint32 j = 0; j < ArrayCount; j++)
-							{
-								AssignedInputs |= (1 << (Location + j));
+						for (uint32 j = 0; j < ArrayCount; j++)
+						{
+							AssignedInputs |= (1 << (Location + j));
 								
-								CCHeaderWriter.WriteInputAttribute(TEXT("in_ATTRIBUTE"), *TypeQualifier, (Location + j), /*bLocationPrefix:*/ false, /*bLocationSuffix:*/ true);
-							}
+							CCHeaderWriter.WriteInputAttribute(TEXT("in_ATTRIBUTE"), *TypeQualifier, (Location + j), /*bLocationPrefix:*/ false, /*bLocationSuffix:*/ true);
 						}
 					}
 				}
diff --git a/Engine/Source/Developer/ShaderCompilerCommon/Private/ShaderConductorContext.cpp b/Engine/Source/Developer/ShaderCompilerCommon/Private/ShaderConductorContext.cpp
--- a/Engine/Source/Developer/ShaderCompilerCommon/Private/ShaderConductorContext.cpp
+++ b/Engine/Source/Developer/ShaderCompilerCommon/Private/ShaderConductorContext.cpp
@@ -839,7 +839,18 @@ namespace CrossCompiler
 	bool FShaderConductorContext::IsIntermediateSpirvOutputVariable(const ANSICHAR* SpirvVariableName)
 	{
 		// This is only true for "temp.var.hullMainRetVal" which is generated by DXC as intermediate output variable to communicate patch constant data in a Hull Shader.
-		return (SpirvVariableName != nullptr && FCStringAnsi::Strcmp(SpirvVariableName, "temp.var.hullMainRetVal") == 0);
+		return (SpirvVariableName != nullptr && FCStringAnsi::Strcmp(SpirvVariableName, FShaderConductorContext::GetIdentifierTable().IntermediateTessControlOutput) == 0);
+	}
+
+	const FShaderConductorIdentifierTable& FShaderConductorContext::GetIdentifierTable()
+	{
+		static const FShaderConductorIdentifierTable IdentifierTable
+		{
+			/*InputAttribute:*/					"in.var.ATTRIBUTE",
+			/*GlobalsUniformBuffer:*/			"$Globals",
+			/*IntermediateTessControlOutput:*/	"temp.var.hullMainRetVal",
+		};
+		return IdentifierTable;
 	}
 
 } // namespace CrossCompiler
diff --git a/Engine/Source/Developer/ShaderCompilerCommon/Private/SpirvReflectCommon.cpp b/Engine/Source/Developer/ShaderCompilerCommon/Private/SpirvReflectCommon.cpp
new file mode 100644
--- /dev/null
+++ b/Engine/Source/Developer/ShaderCompilerCommon/Private/SpirvReflectCommon.cpp
@@ -0,0 +1,212 @@
+// Copyright Epic Games, Inc. All Rights Reserved.
+
+#include "SpirvReflectCommon.h"
+
+void FSpirvReflectBindings::AddDescriptorBinding(SpvReflectDescriptorBinding* InBinding)
+{
+	check(InBinding != nullptr);
+	switch (InBinding->resource_type)
+	{
+	case SPV_REFLECT_RESOURCE_FLAG_SAMPLER:
+		// Gather sampler states (i.e. SamplerState or SamplerComparisonState)
+		check(InBinding->descriptor_type == SPV_REFLECT_DESCRIPTOR_TYPE_SAMPLER);
+		if (InBinding->accessed)
+		{
+			Samplers.Add(InBinding);
+		}
+		break;
+
+	case SPV_REFLECT_RESOURCE_FLAG_CBV:
+		// Gather constant buffers (i.e. cbuffer or ConstantBuffer<T>).
+		check(InBinding->descriptor_type == SPV_REFLECT_DESCRIPTOR_TYPE_UNIFORM_BUFFER);
+		if (InBinding->accessed)
+		{
+			UniformBuffers.Add(InBinding);
+		}
+		break;
+
+	case SPV_REFLECT_RESOURCE_FLAG_SRV:
+		// Gather SRV resources (e.g. Buffer, StructuredBuffer, Texture2D etc.)
+		switch (InBinding->descriptor_type)
+		{
+			/*case SPV_REFLECT_DESCRIPTOR_TYPE_COMBINED_IMAGE_SAMPLER:
+			{
+				if (InBinding->accessed)
+				{
+					Samplers.Add(InBinding);
+				}
+				break;
+			}*/
+		case SPV_REFLECT_DESCRIPTOR_TYPE_SAMPLED_IMAGE:
+			if (InBinding->accessed)
+			{
+				TextureSRVs.Add(InBinding);
+			}
+			break;
+
+		case SPV_REFLECT_DESCRIPTOR_TYPE_UNIFORM_TEXEL_BUFFER:
+			if (InBinding->accessed)
+			{
+				TBufferSRVs.Add(InBinding);
+			}
+			break;
+
+		case SPV_REFLECT_DESCRIPTOR_TYPE_STORAGE_BUFFER:
+			if (InBinding->accessed)
+			{
+				// Storage buffers must always occupy a UAV binding slot
+				SBufferUAVs.Add(InBinding);
+			}
+			break;
+
+		default:
+			// check(false);
+			break;
+		}
+		break;
+
+	case SPV_REFLECT_RESOURCE_FLAG_UAV:
+		if (InBinding->uav_counter_binding)
+		{
+			AtomicCounters.Add(InBinding->uav_counter_binding);
+		}
+
+		switch (InBinding->descriptor_type)
+		{
+		case SPV_REFLECT_DESCRIPTOR_TYPE_STORAGE_IMAGE:
+			TextureUAVs.Add(InBinding);
+			break;
+
+		case SPV_REFLECT_DESCRIPTOR_TYPE_STORAGE_TEXEL_BUFFER:
+			TBufferUAVs.Add(InBinding);
+			break;
+
+		case SPV_REFLECT_DESCRIPTOR_TYPE_STORAGE_BUFFER:
+			if (!AtomicCounters.Contains(InBinding) || InBinding->accessed)
+			{
+				SBufferUAVs.Add(InBinding);
+			}
+			break;
+
+		default:
+			break;
+		}
+		break;
+
+	default:
+		// Gather input attachments (e.g. subpass inputs)
+		switch (InBinding->descriptor_type)
+		{
+		case SPV_REFLECT_DESCRIPTOR_TYPE_INPUT_ATTACHMENT:
+			if (InBinding->accessed)
+			{
+				InputAttachments.Add(InBinding);
+			}
+			break;
+
+		default:
+			break;
+		}
+		break;
+	} // switch
+}
+
+void FSpirvReflectBindings::GatherDescriptorBindings(const spv_reflect::ShaderModule& SpirvReflection)
+{
+	// Enumerate all binding points
+	uint32 NumBindings = 0;
+	SpvReflectResult SpvResult = SpirvReflection.EnumerateDescriptorBindings(&NumBindings, nullptr);
+	check(SpvResult == SPV_REFLECT_RESULT_SUCCESS);
+
+	if (NumBindings > 0)
+	{
+		TArray<SpvReflectDescriptorBinding*> Bindings;
+		Bindings.SetNum(NumBindings);
+		SpvResult = SpirvReflection.EnumerateDescriptorBindings(&NumBindings, Bindings.GetData());
+		check(SpvResult == SPV_REFLECT_RESULT_SUCCESS);
+
+		// Extract all the bindings first so that we process them in order - this lets us assign UAVs before other resources
+		// Which is necessary to match the D3D binding scheme.
+		for (SpvReflectDescriptorBinding* BindingEntry : Bindings)
+		{
+			AddDescriptorBinding(BindingEntry);
+		}
+	}
+}
+
+void FSpirvReflectBindings::GatherInputAttributes(const spv_reflect::ShaderModule& SpirvReflection)
+{
+	// Enumerate all input attributes
+	uint32 NumInputAttributes = 0;
+	SpvReflectResult SpvResult = SpirvReflection.EnumerateEntryPointInputVariables(SpirvReflection.GetEntryPointName(), &NumInputAttributes, nullptr);
+	check(SpvResult == SPV_REFLECT_RESULT_SUCCESS);
+
+	if (NumInputAttributes > 0)
+	{
+		InputAttributes.SetNum(NumInputAttributes);
+		SpvResult = SpirvReflection.EnumerateEntryPointInputVariables(SpirvReflection.GetEntryPointName(), &NumInputAttributes, InputAttributes.GetData());
+		check(SpvResult == SPV_REFLECT_RESULT_SUCCESS);
+	}
+}
+
+void FSpirvReflectBindings::GatherOutputAttributes(const spv_reflect::ShaderModule & SpirvReflection)
+{
+	// Enumerate all output attributes
+	uint32 NumOutputAttributes = 0;
+	SpvReflectResult SpvResult = SpirvReflection.EnumerateEntryPointOutputVariables(SpirvReflection.GetEntryPointName(), &NumOutputAttributes, nullptr);
+	check(SpvResult == SPV_REFLECT_RESULT_SUCCESS);
+
+	if (NumOutputAttributes > 0)
+	{
+		OutputAttributes.SetNum(NumOutputAttributes);
+		SpvResult = SpirvReflection.EnumerateEntryPointOutputVariables(SpirvReflection.GetEntryPointName(), &NumOutputAttributes, OutputAttributes.GetData());
+		check(SpvResult == SPV_REFLECT_RESULT_SUCCESS);
+	}
+}
+
+// Parse the index from a semantic name, e.g. "ATTRIBUTE14" returns 14.
+static bool ParseSemanticIndex(const ANSICHAR* InName, int32& OutSemanticIndex)
+{
+	if (!InName)
+	{
+		return false;
+	}
+
+	for (int32 NameLen = FCStringAnsi::Strlen(InName), Index = NameLen; Index > 0; --Index)
+	{
+		if (!FChar::IsDigit(InName[Index - 1]))
+		{
+			if (Index == NameLen)
+			{
+				// Semantic name does not end with digits
+				return false;
+			}
+			else
+			{
+				// Return suffix numeric starting at previous string position
+				OutSemanticIndex = FCStringAnsi::Atoi(InName + Index);
+				return true;
+			}
+		}
+	}
+
+	return false;
+}
+
+void FSpirvReflectBindings::AssignInputAttributeLocationsBySemanticIndex(spv_reflect::ShaderModule& SpirvReflection, const ANSICHAR* SemanticName)
+{
+	const int32 SemanticNameLen = FCStringAnsi::Strlen(SemanticName);
+	for (SpvReflectInterfaceVariable* InterfaceVar : InputAttributes)
+	{
+		if (InterfaceVar->built_in == -1 && InterfaceVar->name && FCStringAnsi::Strncmp(InterfaceVar->name, SemanticName, SemanticNameLen) == 0)
+		{
+			int32 Location = 0;
+			if (ParseSemanticIndex(InterfaceVar->name, Location))
+			{
+				SpirvReflection.ChangeInputVariableLocation(InterfaceVar, static_cast<uint32>(Location));
+			}
+		}
+	}
+}
+
+
diff --git a/Engine/Source/Developer/ShaderCompilerCommon/Public/ShaderConductorContext.h b/Engine/Source/Developer/ShaderCompilerCommon/Public/ShaderConductorContext.h
--- a/Engine/Source/Developer/ShaderCompilerCommon/Public/ShaderConductorContext.h
+++ b/Engine/Source/Developer/ShaderCompilerCommon/Public/ShaderConductorContext.h
@@ -68,6 +68,19 @@ namespace CrossCompiler
 		TFunction<bool(const FAnsiStringView& VariableName, const FAnsiStringView& TypeName, FString& OutRenamedTypeName)> VariableTypeRenameCallback;
 	};
 
+	/** Container for all special case SPIR-V identifiers generated by ShaderConductor. */
+	struct SHADERCOMPILERCOMMON_API FShaderConductorIdentifierTable
+	{
+		/** Identifier for vertex input attributes: "in.var.ATTRIBUTE". */
+		const ANSICHAR* InputAttribute;
+
+		/** Identifier for globals unifrom buffers: "$Globals". */
+		const ANSICHAR* GlobalsUniformBuffer;
+
+		/** Identifier for the intermediate output variable in a tessellation-control shader. */
+		const ANSICHAR* IntermediateTessControlOutput;
+	};
+
 	/** Wrapper class to handle interface between UE and ShaderConductor. Use to compile HLSL shaders to SPIR-V or high-level languages such as Metal. */
 	class SHADERCOMPILERCOMMON_API FShaderConductorContext
 	{
@@ -133,6 +146,9 @@ namespace CrossCompiler
 		This is only true for a special identifiers generated by DXC to communicate patch constant data in the Hull Shader. */
 		static bool IsIntermediateSpirvOutputVariable(const ANSICHAR* SpirvVariableName);
 
+		/** Returns the table of special identifiers generated by ShaderConductor. */
+		static const FShaderConductorIdentifierTable& GetIdentifierTable();
+
 	public:
 		struct FShaderConductorIntermediates; // Pimpl idiom
 
diff --git a/Engine/Source/Developer/ShaderCompilerCommon/Public/SpirvCommon.h b/Engine/Source/Developer/ShaderCompilerCommon/Public/SpirvCommon.h
new file mode 100644
--- /dev/null
+++ b/Engine/Source/Developer/ShaderCompilerCommon/Public/SpirvCommon.h
@@ -0,0 +1,31 @@
+// Copyright Epic Games, Inc. All Rights Reserved.
+
+#pragma once
+
+#include "CoreMinimal.h"
+
+/** Base structure for SPIR-V modules in the shader backends. */
+struct FSpirv
+{
+	TArray<uint32> Data;
+
+	/** Returns a byte pointer to the SPIR-V data. */
+	FORCEINLINE int8* GetByteData()
+	{
+		return reinterpret_cast<int8*>(Data.GetData());
+	}
+
+	/** Returns a byte pointer to the SPIR-V data. */
+	FORCEINLINE const int8* GetByteData() const
+	{
+		return reinterpret_cast<const int8*>(Data.GetData());
+	}
+
+	/** Returns the size of this SPIR-V module in bytes. */
+	FORCEINLINE int32 GetByteSize() const
+	{
+		return Data.Num() * sizeof(uint32);
+	}
+};
+
+
diff --git a/Engine/Source/Developer/ShaderCompilerCommon/Public/SpirvReflectCommon.h b/Engine/Source/Developer/ShaderCompilerCommon/Public/SpirvReflectCommon.h
new file mode 100644
--- /dev/null
+++ b/Engine/Source/Developer/ShaderCompilerCommon/Public/SpirvReflectCommon.h
@@ -0,0 +1,47 @@
+// Copyright Epic Games, Inc. All Rights Reserved.
+
+#pragma once
+
+#include "CoreMinimal.h"
+
+#if PLATFORM_MAC || PLATFORM_WINDOWS || PLATFORM_LINUX
+THIRD_PARTY_INCLUDES_START
+	#include "spirv_reflect.h"
+THIRD_PARTY_INCLUDES_END
+#endif
+
+
+/** Container structure for all SPIR-V reflection resources and in/out attributes. */
+struct SHADERCOMPILERCOMMON_API FSpirvReflectBindings
+{
+	TArray<SpvReflectInterfaceVariable*> InputAttributes;
+	TArray<SpvReflectInterfaceVariable*> OutputAttributes;
+	TSet<SpvReflectDescriptorBinding*> AtomicCounters;
+	TArray<SpvReflectDescriptorBinding*> InputAttachments; // for subpass inputs
+	TArray<SpvReflectDescriptorBinding*> UniformBuffers;
+	TArray<SpvReflectDescriptorBinding*> Samplers;
+	TArray<SpvReflectDescriptorBinding*> TextureSRVs;
+	TArray<SpvReflectDescriptorBinding*> TextureUAVs;
+	TArray<SpvReflectDescriptorBinding*> TBufferSRVs;
+	TArray<SpvReflectDescriptorBinding*> TBufferUAVs;
+	TArray<SpvReflectDescriptorBinding*> SBufferSRVs;
+	TArray<SpvReflectDescriptorBinding*> SBufferUAVs;
+
+	/** Adds the specified descriptor binding to the corresponding container iff the descriptor is marked as being accessed. */
+	void AddDescriptorBinding(SpvReflectDescriptorBinding* InBinding);
+
+	/** Gathers all descriptor bindings from the specified SPIRV-Reflect module. */
+	void GatherDescriptorBindings(const spv_reflect::ShaderModule& SpirvReflection);
+
+	/** Gathers all input interface variables from the specified SPIRV-Reflect module. */
+	void GatherInputAttributes(const spv_reflect::ShaderModule& SpirvReflection);
+
+	/** Gathers all output interface variables from the specified SPIRV-Reflect module. */
+	void GatherOutputAttributes(const spv_reflect::ShaderModule& SpirvReflection);
+
+	/** Assigns the binding location for all input attributes by its semantic index if their name is equal to 'SemanticName'. Used for vertex shaders where all input attributes have the "ATTRIBUTE" semantic. */
+	void AssignInputAttributeLocationsBySemanticIndex(spv_reflect::ShaderModule& SpirvReflection, const ANSICHAR* SemanticName);
+};
+
+
+
diff --git a/Engine/Source/Developer/ShaderCompilerCommon/ShaderCompilerCommon.Build.cs b/Engine/Source/Developer/ShaderCompilerCommon/ShaderCompilerCommon.Build.cs
--- a/Engine/Source/Developer/ShaderCompilerCommon/ShaderCompilerCommon.Build.cs
+++ b/Engine/Source/Developer/ShaderCompilerCommon/ShaderCompilerCommon.Build.cs
@@ -16,6 +16,7 @@ public class ShaderCompilerCommon : ModuleRules
 		if (Target.Platform == UnrealTargetPlatform.Mac || Target.Platform == UnrealTargetPlatform.Win64 || Target.Platform == UnrealTargetPlatform.Linux)
 		{
 			AddEngineThirdPartyPrivateStaticDependencies(Target, "ShaderConductor");
+			AddEngineThirdPartyPrivateStaticDependencies(Target, "SPIRVReflect");
 		}
 
 		// We only need a header containing definitions
diff --git a/Engine/Source/Developer/ShaderFormatOpenGL/Private/OpenGLShaderCompiler.cpp b/Engine/Source/Developer/ShaderFormatOpenGL/Private/OpenGLShaderCompiler.cpp
--- a/Engine/Source/Developer/ShaderFormatOpenGL/Private/OpenGLShaderCompiler.cpp
+++ b/Engine/Source/Developer/ShaderFormatOpenGL/Private/OpenGLShaderCompiler.cpp
@@ -8,6 +8,7 @@
 #include "Serialization/MemoryWriter.h"
 #include "ShaderFormatOpenGL.h"
 #include "HlslccHeaderWriter.h"
+#include "SpirvReflectCommon.h"
 
 #if PLATFORM_WINDOWS
 #include "Windows/AllowWindowsPlatformTypes.h"
@@ -1372,13 +1373,8 @@ static bool CompileToGlslWithShaderConductor(
 		check(Reflection.GetResult() == SPV_REFLECT_RESULT_SUCCESS);
 
 		SpvReflectResult SPVRResult = SPV_REFLECT_RESULT_NOT_READY;
-		uint32 Count = 0;
-		TArray<SpvReflectDescriptorBinding*> Bindings;
-		TSet<SpvReflectDescriptorBinding*> Counters;
-		TArray<SpvReflectInterfaceVariable*> InputVars;
-		TArray<SpvReflectInterfaceVariable*> OutputVars;
 		TArray<SpvReflectBlockVariable*> ConstantBindings;
-		uint32 GlobalSetId = 32;
+		const uint32 GlobalSetId = 32;
 
 		CrossCompiler::FHlslccHeaderWriter CCHeaderWriter;
 		TArray<FString> Textures;
@@ -1390,254 +1386,137 @@ static bool CompileToGlslWithShaderConductor(
 		uint32 SamplerIndices = 0xffffffff;
 
 		std::map<std::string, std::string> UniformVarNames;
-		Count = 0;
-		SPVRResult = Reflection.EnumerateDescriptorBindings(&Count, nullptr);
-		check(SPVRResult == SPV_REFLECT_RESULT_SUCCESS);
-		Bindings.SetNum(Count);
-		SPVRResult = Reflection.EnumerateDescriptorBindings(&Count, Bindings.GetData());
-		check(SPVRResult == SPV_REFLECT_RESULT_SUCCESS);
-		if (Count > 0)
-		{
-			TArray<SpvReflectDescriptorBinding*> UniformBindings;
-			TArray<SpvReflectDescriptorBinding*> SamplerBindings;
-			TArray<SpvReflectDescriptorBinding*> TextureSRVBindings;
-			TArray<SpvReflectDescriptorBinding*> TextureUAVBindings;
-			TArray<SpvReflectDescriptorBinding*> TBufferSRVBindings;
-			TArray<SpvReflectDescriptorBinding*> TBufferUAVBindings;
-			TArray<SpvReflectDescriptorBinding*> SBufferSRVBindings;
-			TArray<SpvReflectDescriptorBinding*> SBufferUAVBindings;
-
-			// Extract all the bindings first so that we process them in order - this lets us assign UAVs before other resources
-			// Which is necessary to match the D3D binding scheme.
-			for (auto const& Binding : Bindings)
-			{
-				switch (Binding->resource_type)
-				{
-				case SPV_REFLECT_RESOURCE_FLAG_CBV:
-				{
-					check(Binding->descriptor_type == SPV_REFLECT_DESCRIPTOR_TYPE_UNIFORM_BUFFER);
-					if (Binding->accessed)
-					{
-						UniformBindings.Add(Binding);
-					}
-					break;
-				}
-				case SPV_REFLECT_RESOURCE_FLAG_SAMPLER:
-				{
-					check(Binding->descriptor_type == SPV_REFLECT_DESCRIPTOR_TYPE_SAMPLER);
-					if (Binding->accessed)
-					{
-						SamplerBindings.Add(Binding);
-					}
-					break;
-				}
-				case SPV_REFLECT_RESOURCE_FLAG_SRV:
-				{
-					switch (Binding->descriptor_type)
-					{
-					case SPV_REFLECT_DESCRIPTOR_TYPE_SAMPLED_IMAGE:
-					{
-						if (Binding->accessed)
-						{
-							TextureSRVBindings.Add(Binding);
-						}
-						break;
-					}
-					case SPV_REFLECT_DESCRIPTOR_TYPE_UNIFORM_TEXEL_BUFFER:
-					{
-						if (Binding->accessed)
-						{
-							TBufferSRVBindings.Add(Binding);
-						}
-						break;
-					}
-					case SPV_REFLECT_DESCRIPTOR_TYPE_STORAGE_BUFFER:
-					{
-						if (Binding->accessed)
-						{
-							SBufferSRVBindings.Add(Binding);
-						}
-						break;
-					}
-					default:
-					{
-						// check(false);
-						break;
-					}
-					}
-					break;
-				}
-				case SPV_REFLECT_RESOURCE_FLAG_UAV:
-				{
-					if (Binding->uav_counter_binding)
-					{
-						Counters.Add(Binding->uav_counter_binding);
-					}
 
-					switch (Binding->descriptor_type)
-					{
-					case SPV_REFLECT_DESCRIPTOR_TYPE_STORAGE_IMAGE:
-					{
-						TextureUAVBindings.Add(Binding);
-						break;
-					}
-					case SPV_REFLECT_DESCRIPTOR_TYPE_STORAGE_TEXEL_BUFFER:
-					{
-						TBufferUAVBindings.Add(Binding);
-						break;
-					}
-					case SPV_REFLECT_DESCRIPTOR_TYPE_STORAGE_BUFFER:
-					{
-						if (!Counters.Contains(Binding) || Binding->accessed)
-						{
-							SBufferUAVBindings.Add(Binding);
-						}
-						break;
-					}
-					default:
-					{
-						// check(false);
-						break;
-					}
-					}
-					break;
-				}
-				default:
-				{
-					// check(false);
-					break;
-				}
-				}
-			}
+		FSpirvReflectBindings ReflectionBindings;
+		ReflectionBindings.GatherDescriptorBindings(Reflection);
 
-			for (auto const& Binding : TBufferUAVBindings)
-			{
-				check(UAVIndices);
-				uint32 Index = FPlatformMath::CountTrailingZeros(UAVIndices);
+		for (auto const& Binding : ReflectionBindings.TBufferUAVs)
+		{
+			check(UAVIndices);
+			uint32 Index = FPlatformMath::CountTrailingZeros(UAVIndices);
 
-				// UAVs always claim all slots so we don't have conflicts as D3D expects 0-7
-				BufferIndices &= ~(1 << Index);
-				TextureIndices &= ~(1llu << uint64(Index));
-				UAVIndices &= ~(1 << Index);
+			// UAVs always claim all slots so we don't have conflicts as D3D expects 0-7
+			BufferIndices &= ~(1 << Index);
+			TextureIndices &= ~(1llu << uint64(Index));
+			UAVIndices &= ~(1 << Index);
 
-				CCHeaderWriter.WriteUAV(UTF8_TO_TCHAR(Binding->name), Index);
+			CCHeaderWriter.WriteUAV(UTF8_TO_TCHAR(Binding->name), Index);
 
-				SPVRResult = Reflection.ChangeDescriptorBindingNumbers(Binding, Index, GlobalSetId);
-				check(SPVRResult == SPV_REFLECT_RESULT_SUCCESS);
-			}
+			SPVRResult = Reflection.ChangeDescriptorBindingNumbers(Binding, Index, GlobalSetId);
+			check(SPVRResult == SPV_REFLECT_RESULT_SUCCESS);
+		}
 
-			for (auto const& Binding : SBufferUAVBindings)
-			{
-				check(UAVIndices);
-				uint32 Index = FPlatformMath::CountTrailingZeros(UAVIndices);
+		for (auto const& Binding : ReflectionBindings.SBufferUAVs)
+		{
+			check(UAVIndices);
+			uint32 Index = FPlatformMath::CountTrailingZeros(UAVIndices);
 
-				// UAVs always claim all slots so we don't have conflicts as D3D expects 0-7
-				BufferIndices &= ~(1 << Index);
-				TextureIndices &= ~(1llu << uint64(Index));
-				UAVIndices &= ~(1 << Index);
+			// UAVs always claim all slots so we don't have conflicts as D3D expects 0-7
+			BufferIndices &= ~(1 << Index);
+			TextureIndices &= ~(1llu << uint64(Index));
+			UAVIndices &= ~(1 << Index);
 
-				CCHeaderWriter.WriteUAV(UTF8_TO_TCHAR(Binding->name), Index);
+			CCHeaderWriter.WriteUAV(UTF8_TO_TCHAR(Binding->name), Index);
 
-				SPVRResult = Reflection.ChangeDescriptorBindingNumbers(Binding, Index, GlobalSetId);
-				check(SPVRResult == SPV_REFLECT_RESULT_SUCCESS);
-			}
+			SPVRResult = Reflection.ChangeDescriptorBindingNumbers(Binding, Index, GlobalSetId);
+			check(SPVRResult == SPV_REFLECT_RESULT_SUCCESS);
+		}
 
-			for (auto const& Binding : TextureUAVBindings)
-			{
-				check(UAVIndices);
-				uint32 Index = FPlatformMath::CountTrailingZeros(UAVIndices);
+		for (auto const& Binding : ReflectionBindings.TextureUAVs)
+		{
+			check(UAVIndices);
+			uint32 Index = FPlatformMath::CountTrailingZeros(UAVIndices);
 
-				// UAVs always claim all slots so we don't have conflicts as D3D expects 0-7
-				// For texture2d this allows us to emulate atomics with buffers
-				BufferIndices &= ~(1 << Index);
-				TextureIndices &= ~(1llu << uint64(Index));
-				UAVIndices &= ~(1 << Index);
+			// UAVs always claim all slots so we don't have conflicts as D3D expects 0-7
+			// For texture2d this allows us to emulate atomics with buffers
+			BufferIndices &= ~(1 << Index);
+			TextureIndices &= ~(1llu << uint64(Index));
+			UAVIndices &= ~(1 << Index);
 
-				CCHeaderWriter.WriteUAV(UTF8_TO_TCHAR(Binding->name), Index);
+			CCHeaderWriter.WriteUAV(UTF8_TO_TCHAR(Binding->name), Index);
 
-				SPVRResult = Reflection.ChangeDescriptorBindingNumbers(Binding, Index, GlobalSetId);
-				check(SPVRResult == SPV_REFLECT_RESULT_SUCCESS);
-			}
+			SPVRResult = Reflection.ChangeDescriptorBindingNumbers(Binding, Index, GlobalSetId);
+			check(SPVRResult == SPV_REFLECT_RESULT_SUCCESS);
+		}
 
-			for (auto const& Binding : TBufferSRVBindings)
-			{
-				check(TextureIndices);
-				uint32 Index = FPlatformMath::CountTrailingZeros(TextureIndices);
+		for (auto const& Binding : ReflectionBindings.TBufferSRVs)
+		{
+			check(TextureIndices);
+			uint32 Index = FPlatformMath::CountTrailingZeros(TextureIndices);
 
-				// No support for 3-component types in dxc/SPIRV/MSL - need to expose my workarounds there too
-				BufferIndices &= ~(1 << Index);
-				TextureIndices &= ~(1llu << uint64(Index));
+			// No support for 3-component types in dxc/SPIRV/MSL - need to expose my workarounds there too
+			BufferIndices &= ~(1 << Index);
+			TextureIndices &= ~(1llu << uint64(Index));
 
-				Textures.Add(UTF8_TO_TCHAR(Binding->name));
+			Textures.Add(UTF8_TO_TCHAR(Binding->name));
 
-				SPVRResult = Reflection.ChangeDescriptorBindingNumbers(Binding, Index, GlobalSetId);
-				check(SPVRResult == SPV_REFLECT_RESULT_SUCCESS);
-			}
+			SPVRResult = Reflection.ChangeDescriptorBindingNumbers(Binding, Index, GlobalSetId);
+			check(SPVRResult == SPV_REFLECT_RESULT_SUCCESS);
+		}
 
-			for (auto const& Binding : SBufferSRVBindings)
-			{
-				check(BufferIndices);
-				uint32 Index = FPlatformMath::CountTrailingZeros(BufferIndices);
+		for (auto const& Binding : ReflectionBindings.SBufferSRVs)
+		{
+			check(BufferIndices);
+			uint32 Index = FPlatformMath::CountTrailingZeros(BufferIndices);
 
-				BufferIndices &= ~(1 << Index);
+			BufferIndices &= ~(1 << Index);
 
-				SPVRResult = Reflection.ChangeDescriptorBindingNumbers(Binding, Index, GlobalSetId);
-				check(SPVRResult == SPV_REFLECT_RESULT_SUCCESS);
-			}
+			SPVRResult = Reflection.ChangeDescriptorBindingNumbers(Binding, Index, GlobalSetId);
+			check(SPVRResult == SPV_REFLECT_RESULT_SUCCESS);
+		}
 
-			for (auto const& Binding : UniformBindings)
-			{
-				check(UBOIndices);
-				uint32 Index = FPlatformMath::CountTrailingZeros(UBOIndices);
-				UBOIndices &= ~(1 << Index);
+		for (auto const& Binding : ReflectionBindings.UniformBuffers)
+		{
+			check(UBOIndices);
+			uint32 Index = FPlatformMath::CountTrailingZeros(UBOIndices);
+			UBOIndices &= ~(1 << Index);
 
-				// Global uniform buffer - handled specially as we care about the internal layout
-				if (strstr(Binding->name, "$Globals"))
-				{
-					for (uint32 i = 0; i < Binding->block.member_count; i++)
-					{
-						CCHeaderWriter.WritePackedGlobal(Binding->block.members[i]);
-					}
-				}
-				else
+			// Global uniform buffer - handled specially as we care about the internal layout
+			if (strstr(Binding->name, "$Globals"))
+			{
+				for (uint32 i = 0; i < Binding->block.member_count; i++)
 				{
-					std::string OldName = "type_";
-					OldName += Binding->name;
-					std::string NewName = FrequencyPrefix;
-					NewName += "b";
-					NewName += std::to_string(Index);
-					UniformVarNames[OldName] = NewName;
-					// Regular uniform buffer - we only care about the binding index
-					CCHeaderWriter.WriteUniformBlock(UTF8_TO_TCHAR(Binding->name), Index);
+					CCHeaderWriter.WritePackedGlobal(Binding->block.members[i]);
 				}
-
-				SPVRResult = Reflection.ChangeDescriptorBindingNumbers(Binding, Index, GlobalSetId);
-				check(SPVRResult == SPV_REFLECT_RESULT_SUCCESS);
 			}
-
-			for (auto const& Binding : TextureSRVBindings)
+			else
 			{
-				check(TextureIndices);
-				uint32 Index = FPlatformMath::CountTrailingZeros64(TextureIndices);
-				TextureIndices &= ~(1llu << uint64(Index));
+				std::string OldName = "type_";
+				OldName += Binding->name;
+				std::string NewName = FrequencyPrefix;
+				NewName += "b";
+				NewName += std::to_string(Index);
+				UniformVarNames[OldName] = NewName;
+				// Regular uniform buffer - we only care about the binding index
+				CCHeaderWriter.WriteUniformBlock(UTF8_TO_TCHAR(Binding->name), Index);
+			}
+
+			SPVRResult = Reflection.ChangeDescriptorBindingNumbers(Binding, Index, GlobalSetId);
+			check(SPVRResult == SPV_REFLECT_RESULT_SUCCESS);
+		}
 
-				Textures.Add(UTF8_TO_TCHAR(Binding->name));
+		for (auto const& Binding : ReflectionBindings.TextureSRVs)
+		{
+			check(TextureIndices);
+			uint32 Index = FPlatformMath::CountTrailingZeros64(TextureIndices);
+			TextureIndices &= ~(1llu << uint64(Index));
 
-				SPVRResult = Reflection.ChangeDescriptorBindingNumbers(Binding, Index, GlobalSetId);
-				check(SPVRResult == SPV_REFLECT_RESULT_SUCCESS);
-			}
+			Textures.Add(UTF8_TO_TCHAR(Binding->name));
 
-			for (auto const& Binding : SamplerBindings)
-			{
-				check(SamplerIndices);
-				uint32 Index = FPlatformMath::CountTrailingZeros(SamplerIndices);
-				SamplerIndices &= ~(1 << Index);
+			SPVRResult = Reflection.ChangeDescriptorBindingNumbers(Binding, Index, GlobalSetId);
+			check(SPVRResult == SPV_REFLECT_RESULT_SUCCESS);
+		}
 
-				Samplers.Add(UTF8_TO_TCHAR(Binding->name));
+		for (auto const& Binding : ReflectionBindings.Samplers)
+		{
+			check(SamplerIndices);
+			uint32 Index = FPlatformMath::CountTrailingZeros(SamplerIndices);
+			SamplerIndices &= ~(1 << Index);
 
-				SPVRResult = Reflection.ChangeDescriptorBindingNumbers(Binding, Index, GlobalSetId);
-				check(SPVRResult == SPV_REFLECT_RESULT_SUCCESS);
-			}
+			Samplers.Add(UTF8_TO_TCHAR(Binding->name));
+
+			SPVRResult = Reflection.ChangeDescriptorBindingNumbers(Binding, Index, GlobalSetId);
+			check(SPVRResult == SPV_REFLECT_RESULT_SUCCESS);
 		}
 
 		TArray<std::string> GlobalRemap;
@@ -1645,7 +1524,7 @@ static bool CompileToGlslWithShaderConductor(
 		TMap<FString, uint32> GlobalOffsets;
 
 		{
-			Count = 0;
+			uint32 Count = 0;
 			SPVRResult = Reflection.EnumeratePushConstantBlocks(&Count, nullptr);
 			check(SPVRResult == SPV_REFLECT_RESULT_SUCCESS);
 			ConstantBindings.SetNum(Count);
@@ -1777,144 +1656,46 @@ static bool CompileToGlslWithShaderConductor(
 		TArray<FString> OutputVarNames;
 
 		{
-			Count = 0;
-			SPVRResult = Reflection.EnumerateOutputVariables(&Count, nullptr);
-			check(SPVRResult == SPV_REFLECT_RESULT_SUCCESS);
-			OutputVars.SetNum(Count);
-			SPVRResult = Reflection.EnumerateOutputVariables(&Count, OutputVars.GetData());
-			check(SPVRResult == SPV_REFLECT_RESULT_SUCCESS);
-			if (Count > 0)
-			{
-				uint32 AssignedInputs = 0;
+			uint32 AssignedInputs = 0;
 
-				for (auto const& Var : OutputVars)
+			ReflectionBindings.GatherOutputAttributes(Reflection);
+			for (SpvReflectInterfaceVariable* Var : ReflectionBindings.OutputAttributes)
+			{
+				if (Var->storage_class == SpvStorageClassOutput && Var->built_in == -1 && !CrossCompiler::FShaderConductorContext::IsIntermediateSpirvOutputVariable(Var->name))
 				{
-					if (Var->storage_class == SpvStorageClassOutput && Var->built_in == -1 && !CrossCompiler::FShaderConductorContext::IsIntermediateSpirvOutputVariable(Var->name))
+					if (Frequency == HSF_PixelShader && strstr(Var->name, "SV_Target"))
 					{
-						if (Frequency == HSF_PixelShader && strstr(Var->name, "SV_Target"))
-						{
-							FString TypeQualifier;
-
-							auto const type = *Var->type_description;
-							uint32_t masked_type = type.type_flags & 0xF;
+						FString TypeQualifier;
 
-							switch (masked_type) {
-							default: checkf(false, TEXT("unsupported component type %d"), masked_type); break;
-							case SPV_REFLECT_TYPE_FLAG_BOOL: TypeQualifier = TEXT("b"); break;
-							case SPV_REFLECT_TYPE_FLAG_INT: TypeQualifier = (type.traits.numeric.scalar.signedness ? TEXT("i") : TEXT("u")); break;
-							case SPV_REFLECT_TYPE_FLAG_FLOAT: TypeQualifier = (type.traits.numeric.scalar.width == 32 ? TEXT("f") : TEXT("h")); break;
-							}
+						auto const type = *Var->type_description;
+						uint32_t masked_type = type.type_flags & 0xF;
 
-							if (type.type_flags & SPV_REFLECT_TYPE_FLAG_MATRIX)
-							{
-								TypeQualifier += FString::Printf(TEXT("%d%d"), type.traits.numeric.matrix.row_count, type.traits.numeric.matrix.column_count);
-							}
-							else if (type.type_flags & SPV_REFLECT_TYPE_FLAG_VECTOR)
-							{
-								TypeQualifier += FString::Printf(TEXT("%d"), type.traits.numeric.vector.component_count);
-							}
-							else
-							{
-								TypeQualifier += TEXT("1");
-							}
+						switch (masked_type) {
+						default: checkf(false, TEXT("unsupported component type %d"), masked_type); break;
+						case SPV_REFLECT_TYPE_FLAG_BOOL: TypeQualifier = TEXT("b"); break;
+						case SPV_REFLECT_TYPE_FLAG_INT: TypeQualifier = (type.traits.numeric.scalar.signedness ? TEXT("i") : TEXT("u")); break;
+						case SPV_REFLECT_TYPE_FLAG_FLOAT: TypeQualifier = (type.traits.numeric.scalar.width == 32 ? TEXT("f") : TEXT("h")); break;
+						}
 
-							FString Name = ANSI_TO_TCHAR(Var->name);
-							Name.ReplaceInline(TEXT("."), TEXT("_"));
-							OutputVarNames.Add(Name);
-							CCHeaderWriter.WriteOutputAttribute(TEXT("out_Target"), *TypeQualifier, Var->location, /*bLocationPrefix:*/ true, /*bLocationSuffix:*/ true);
+						if (type.type_flags & SPV_REFLECT_TYPE_FLAG_MATRIX)
+						{
+							TypeQualifier += FString::Printf(TEXT("%d%d"), type.traits.numeric.matrix.row_count, type.traits.numeric.matrix.column_count);
+						}
+						else if (type.type_flags & SPV_REFLECT_TYPE_FLAG_VECTOR)
+						{
+							TypeQualifier += FString::Printf(TEXT("%d"), type.traits.numeric.vector.component_count);
 						}
 						else
 						{
-							unsigned Location = Var->location;
-							unsigned SemanticIndex = Location;
-							check(Var->semantic);
-							unsigned i = (unsigned)strlen(Var->semantic);
-							check(i);
-							while (isdigit((unsigned char)(Var->semantic[i - 1])))
-							{
-								i--;
-							}
-							if (i < strlen(Var->semantic))
-							{
-								SemanticIndex = (unsigned)atoi(Var->semantic + i);
-								if (Location != SemanticIndex)
-								{
-									Location = SemanticIndex;
-								}
-							}
-
-							while ((1 << Location) & AssignedInputs)
-							{
-								Location++;
-							}
-
-							if (Location != Var->location)
-							{
-								SPVRResult = Reflection.ChangeOutputVariableLocation(Var, Location);
-								check(SPVRResult == SPV_REFLECT_RESULT_SUCCESS);
-							}
-
-							uint32 ArrayCount = 1;
-							for (uint32 Dim = 0; Dim < Var->array.dims_count; Dim++)
-							{
-								ArrayCount *= Var->array.dims[Dim];
-							}
-
-							FString TypeQualifier;
-
-							auto const type = *Var->type_description;
-							uint32_t masked_type = type.type_flags & 0xF;
-
-							switch (masked_type) {
-							default: checkf(false, TEXT("unsupported component type %d"), masked_type); break;
-							case SPV_REFLECT_TYPE_FLAG_BOOL: TypeQualifier = TEXT("b"); break;
-							case SPV_REFLECT_TYPE_FLAG_INT: TypeQualifier = (type.traits.numeric.scalar.signedness ? TEXT("i") : TEXT("u")); break;
-							case SPV_REFLECT_TYPE_FLAG_FLOAT: TypeQualifier = (type.traits.numeric.scalar.width == 32 ? TEXT("f") : TEXT("h")); break;
-							}
-
-							if (type.type_flags & SPV_REFLECT_TYPE_FLAG_MATRIX)
-							{
-								TypeQualifier += FString::Printf(TEXT("%d%d"), type.traits.numeric.matrix.row_count, type.traits.numeric.matrix.column_count);
-							}
-							else if (type.type_flags & SPV_REFLECT_TYPE_FLAG_VECTOR)
-							{
-								TypeQualifier += FString::Printf(TEXT("%d"), type.traits.numeric.vector.component_count);
-							}
-							else
-							{
-								TypeQualifier += TEXT("1");
-							}
-
-							for (uint32 j = 0; j < ArrayCount; j++)
-							{
-								AssignedInputs |= (1 << (Location + j));
-							}
-
-							FString Name = ANSI_TO_TCHAR(Var->name);
-							Name.ReplaceInline(TEXT("."), TEXT("_"));
-							CCHeaderWriter.WriteOutputAttribute(*Name, *TypeQualifier, Location, /*bLocationPrefix:*/ true, /*bLocationSuffix:*/ false);
+							TypeQualifier += TEXT("1");
 						}
-					}
-				}
-			}
-		}
-
-		TArray<FString> InputVarNames;
-
-		{
-			Count = 0;
-			SPVRResult = Reflection.EnumerateInputVariables(&Count, nullptr);
-			check(SPVRResult == SPV_REFLECT_RESULT_SUCCESS);
-			InputVars.SetNum(Count);
-			SPVRResult = Reflection.EnumerateInputVariables(&Count, InputVars.GetData());
-			check(SPVRResult == SPV_REFLECT_RESULT_SUCCESS);
-			if (Count > 0)
-			{
-				uint32 AssignedInputs = 0;
 
-				for (auto const& Var : InputVars)
-				{
-					if (Var->storage_class == SpvStorageClassInput && Var->built_in == -1)
+						FString Name = ANSI_TO_TCHAR(Var->name);
+						Name.ReplaceInline(TEXT("."), TEXT("_"));
+						OutputVarNames.Add(Name);
+						CCHeaderWriter.WriteOutputAttribute(TEXT("out_Target"), *TypeQualifier, Var->location, /*bLocationPrefix:*/ true, /*bLocationSuffix:*/ true);
+					}
+					else
 					{
 						unsigned Location = Var->location;
 						unsigned SemanticIndex = Location;
@@ -1941,7 +1722,7 @@ static bool CompileToGlslWithShaderConductor(
 
 						if (Location != Var->location)
 						{
-							SPVRResult = Reflection.ChangeInputVariableLocation(Var, Location);
+							SPVRResult = Reflection.ChangeOutputVariableLocation(Var, Location);
 							check(SPVRResult == SPV_REFLECT_RESULT_SUCCESS);
 						}
 
@@ -1983,9 +1764,91 @@ static bool CompileToGlslWithShaderConductor(
 
 						FString Name = ANSI_TO_TCHAR(Var->name);
 						Name.ReplaceInline(TEXT("."), TEXT("_"));
-						InputVarNames.Add(Name);
-						CCHeaderWriter.WriteInputAttribute(*Name, *TypeQualifier, Location, /*bLocationPrefix:*/ true, /*bLocationSuffix:*/ false);
+						CCHeaderWriter.WriteOutputAttribute(*Name, *TypeQualifier, Location, /*bLocationPrefix:*/ true, /*bLocationSuffix:*/ false);
+					}
+				}
+			}
+		}
+
+		TArray<FString> InputVarNames;
+
+		{
+			uint32 AssignedInputs = 0;
+
+			ReflectionBindings.GatherInputAttributes(Reflection);
+			for (SpvReflectInterfaceVariable* Var : ReflectionBindings.InputAttributes)
+			{
+				if (Var->storage_class == SpvStorageClassInput && Var->built_in == -1)
+				{
+					unsigned Location = Var->location;
+					unsigned SemanticIndex = Location;
+					check(Var->semantic);
+					unsigned i = (unsigned)strlen(Var->semantic);
+					check(i);
+					while (isdigit((unsigned char)(Var->semantic[i - 1])))
+					{
+						i--;
 					}
+					if (i < strlen(Var->semantic))
+					{
+						SemanticIndex = (unsigned)atoi(Var->semantic + i);
+						if (Location != SemanticIndex)
+						{
+							Location = SemanticIndex;
+						}
+					}
+
+					while ((1 << Location) & AssignedInputs)
+					{
+						Location++;
+					}
+
+					if (Location != Var->location)
+					{
+						SPVRResult = Reflection.ChangeInputVariableLocation(Var, Location);
+						check(SPVRResult == SPV_REFLECT_RESULT_SUCCESS);
+					}
+
+					uint32 ArrayCount = 1;
+					for (uint32 Dim = 0; Dim < Var->array.dims_count; Dim++)
+					{
+						ArrayCount *= Var->array.dims[Dim];
+					}
+
+					FString TypeQualifier;
+
+					auto const type = *Var->type_description;
+					uint32_t masked_type = type.type_flags & 0xF;
+
+					switch (masked_type) {
+					default: checkf(false, TEXT("unsupported component type %d"), masked_type); break;
+					case SPV_REFLECT_TYPE_FLAG_BOOL: TypeQualifier = TEXT("b"); break;
+					case SPV_REFLECT_TYPE_FLAG_INT: TypeQualifier = (type.traits.numeric.scalar.signedness ? TEXT("i") : TEXT("u")); break;
+					case SPV_REFLECT_TYPE_FLAG_FLOAT: TypeQualifier = (type.traits.numeric.scalar.width == 32 ? TEXT("f") : TEXT("h")); break;
+					}
+
+					if (type.type_flags & SPV_REFLECT_TYPE_FLAG_MATRIX)
+					{
+						TypeQualifier += FString::Printf(TEXT("%d%d"), type.traits.numeric.matrix.row_count, type.traits.numeric.matrix.column_count);
+					}
+					else if (type.type_flags & SPV_REFLECT_TYPE_FLAG_VECTOR)
+					{
+						TypeQualifier += FString::Printf(TEXT("%d"), type.traits.numeric.vector.component_count);
+					}
+					else
+					{
+						TypeQualifier += TEXT("1");
+					}
+
+					for (uint32 j = 0; j < ArrayCount; j++)
+					{
+						AssignedInputs |= (1 << (Location + j));
+					}
+
+					FString Name = ANSI_TO_TCHAR(Var->name);
+					Name.ReplaceInline(TEXT("."), TEXT("_"));
+					InputVarNames.Add(Name);
+					CCHeaderWriter.WriteInputAttribute(*Name, *TypeQualifier, Location, /*bLocationPrefix:*/ true, /*bLocationSuffix:*/ false);
 				}
 			}
 		}
diff --git a/Engine/Source/Developer/VulkanShaderFormat/Private/VulkanGlslangBridge.cpp b/Engine/Source/Developer/VulkanShaderFormat/Private/VulkanGlslangBridge.cpp
--- a/Engine/Source/Developer/VulkanShaderFormat/Private/VulkanGlslangBridge.cpp
+++ b/Engine/Source/Developer/VulkanShaderFormat/Private/VulkanGlslangBridge.cpp
@@ -166,7 +166,7 @@ static EShLanguage GetStage(EHlslShaderFrequency Frequency)
 	return EShLangCount;
 }
 
-static void ComputeMovableWordIndices(FSpirv& Spirv)
+static void ComputeMovableWordIndices(FVulkanSpirv& Spirv)
 {
 	// SPIRV Header
 	const uint32* PtrStart = Spirv.Data.GetData();
@@ -311,7 +311,7 @@ static void ComputeMovableWordIndices(FSpirv& Spirv)
 				const FString* FoundTypeName = Names.Find(TypePointer);
 				if (FoundTypeName && FoundTypeName->Len() > 0)
 				{
-					FSpirv::FEntry* FoundEntry = Spirv.GetEntry(*FoundTypeName);
+					FVulkanSpirv::FEntry* FoundEntry = Spirv.GetEntry(*FoundTypeName);
 					checkf(FoundEntry, TEXT("cannot find variable entry in SPIR-V module: %s"), **FoundTypeName);
 					FDecorations& FoundDecorations = Decorations.FindChecked(VariableId);
 					FoundEntry->Binding = FoundDecorations.BindingIndex;
@@ -323,7 +323,7 @@ static void ComputeMovableWordIndices(FSpirv& Spirv)
 			else
 			{
 				// Standalone global var
-				FSpirv::FEntry* FoundEntry = Spirv.GetEntry(*FoundVariableName);
+				FVulkanSpirv::FEntry* FoundEntry = Spirv.GetEntry(*FoundVariableName);
 				checkf(FoundEntry, TEXT("Entry name not found in SPIR-V module: %s"), *(*FoundVariableName));
 				FDecorations& FoundDecorations = Decorations.FindChecked(VariableId);
 				FoundEntry->Binding = FoundDecorations.BindingIndex;
@@ -335,14 +335,14 @@ static void ComputeMovableWordIndices(FSpirv& Spirv)
 	}
 }
 
-static void PatchSpirvEntryPoint(FSpirv& OutSpirv, uint32 OffsetToName)
+static void PatchSpirvEntryPoint(FVulkanSpirv& OutSpirv, uint32 OffsetToName)
 {
 	char* EntryPointName = (char*)(OutSpirv.Data.GetData() + OffsetToName);
 	check(!FCStringAnsi::Strcmp(EntryPointName, "main_00000000_00000000"));
 	FCStringAnsi::Sprintf(EntryPointName, "main_%0.8x_%0.8x", OutSpirv.Data.Num() * sizeof(uint32), OutSpirv.CRC);
 };
 
-bool PatchSpirvReflectionEntriesAndEntryPoint(FSpirv& OutSpirv)
+bool PatchSpirvReflectionEntriesAndEntryPoint(FVulkanSpirv& OutSpirv)
 {
 	// Re-compute movable word indices and update CRC code
 	ComputeMovableWordIndices(OutSpirv);
@@ -354,7 +354,7 @@ bool PatchSpirvReflectionEntriesAndEntryPoint(FSpirv& OutSpirv)
 	return true;
 }
 
-bool GenerateSpirv(const ANSICHAR* Source, FCompilerInfo& CompilerInfo, FString& OutErrors, const FString& DumpDebugInfoPath, FSpirv& OutSpirv)
+bool GenerateSpirv(const ANSICHAR* Source, FCompilerInfo& CompilerInfo, FString& OutErrors, const FString& DumpDebugInfoPath, FVulkanSpirv& OutSpirv)
 {
 	glslang::TProgram* Program = new glslang::TProgram;
 
@@ -422,7 +422,7 @@ bool GenerateSpirv(const ANSICHAR* Source, FCompilerInfo& CompilerInfo, FString&
 			{
 				Binding = Program->getUniformBinding(Index);
 			}
-			FSpirv::FEntry Entry(Name, Binding);
+			FVulkanSpirv::FEntry Entry(Name, Binding);
 			OutSpirv.ReflectionInfo.Add(Entry);
 		}
 
@@ -440,7 +440,7 @@ bool GenerateSpirv(const ANSICHAR* Source, FCompilerInfo& CompilerInfo, FString&
 			{
 				Binding = Program->getUniformBinding(Index);
 			}
-			FSpirv::FEntry Entry(Name, Binding);
+			FVulkanSpirv::FEntry Entry(Name, Binding);
 			OutSpirv.ReflectionInfo.Add(Entry);
 		}
 
diff --git a/Engine/Source/Developer/VulkanShaderFormat/Private/VulkanShaderCompiler.cpp b/Engine/Source/Developer/VulkanShaderFormat/Private/VulkanShaderCompiler.cpp
--- a/Engine/Source/Developer/VulkanShaderFormat/Private/VulkanShaderCompiler.cpp
+++ b/Engine/Source/Developer/VulkanShaderFormat/Private/VulkanShaderCompiler.cpp
@@ -7,10 +7,10 @@
 #include "ShaderCompilerCommon.h"
 #include "HlslccHeaderWriter.h"
 #include "hlslcc.h"
+#include "SpirvReflectCommon.h"
 
 #if PLATFORM_MAC || PLATFORM_WINDOWS || PLATFORM_LINUX
 THIRD_PARTY_INCLUDES_START
-	#include "spirv_reflect.h"
 	#include "SPIRV/GlslangToSpv.h"
 	#include "SPIRV/doc.h"
 	#include "SPIRV/disassemble.h"
@@ -398,7 +398,7 @@ static int32 AddGlobal(FOLDVulkanCodeHeader& OLDHeader,
 						const CrossCompiler::FHlslccHeader& CCHeader,
 						const FString& ParameterName,
 						uint16 BindingIndex,
-						const FSpirv& Spirv,
+						const FVulkanSpirv& Spirv,
 						FVulkanShaderHeader& OutHeader,
 						const TArray<FString>& GlobalNames,
 						TArray<FPatchType>& OutTypePatch,
@@ -409,7 +409,7 @@ static int32 AddGlobal(FOLDVulkanCodeHeader& OLDHeader,
 	check(GlobalNames[HeaderGlobalIndex] == ParameterName);
 
 	FVulkanShaderHeader::FGlobalInfo& GlobalInfo = OutHeader.Globals[HeaderGlobalIndex];
-	const FSpirv::FEntry* Entry = Spirv.GetEntry(ParameterName);
+	const FVulkanSpirv::FEntry* Entry = Spirv.GetEntry(ParameterName);
 	bool bIsCombinedSampler = false;
 	if (Entry)
 	{
@@ -470,7 +470,7 @@ static int32 AddGlobalForUBEntry(FOLDVulkanCodeHeader& OLDHeader,
 									const CrossCompiler::FHlslccHeader& CCHeader,
 									const FString& ParameterName,
 									uint16 BindingIndex,
-									const FSpirv& Spirv,
+									const FVulkanSpirv& Spirv,
 									const TArray<FString>&
 									GlobalNames,
 									EUniformBufferBaseType UBEntryType,
@@ -503,7 +503,7 @@ static void AddUBResources(FOLDVulkanCodeHeader& OLDHeader,
 							const TArray<uint32>& BindingArray,
 							const FVulkanBindingTable& BindingTable,
 							const TArray<VkDescriptorType>& DescriptorTypes,
-							const FSpirv& Spirv,
+							const FVulkanSpirv& Spirv,
 							const CrossCompiler::FHlslccHeader& CCHeader,
 							FVulkanShaderHeader::FUniformBufferInfo& OutUBInfo,
 							FVulkanShaderHeader& OutHeader,
@@ -555,7 +555,7 @@ static void AddUniformBuffer(FOLDVulkanCodeHeader& OLDHeader,
 	const FVulkanBindingTable& BindingTable,
 	const FShaderCompilerInput& ShaderInput,
 	const CrossCompiler::FHlslccHeader& CCHeader,
-	const FSpirv& Spirv,
+	const FVulkanSpirv& Spirv,
 	const FString& UBName,
 	uint16 BindingIndex,
 	FShaderParameterMap& InOutParameterMap,
@@ -583,7 +583,7 @@ static void AddUniformBuffer(FOLDVulkanCodeHeader& OLDHeader,
 #if VULKAN_ENABLE_BINDING_DEBUG_NAMES
 	UBInfo.DebugName = UBName;
 #endif
-	const FSpirv::FEntry* Entry = Spirv.GetEntry(UBName);
+	const FVulkanSpirv::FEntry* Entry = Spirv.GetEntry(UBName);
 	if (Entry)
 	{
 		checkf(bUseRealUBs, TEXT("Internal error: Emulated should NOT have a binding!"));
@@ -864,7 +864,7 @@ static void PrepareGlobals(const FVulkanBindingTable& BindingTable, const FVulka
 
 static void ConvertToNEWHeader(FOLDVulkanCodeHeader& OLDHeader,
 	const FVulkanBindingTable& BindingTable,
-	const FSpirv& Spirv,
+	const FVulkanSpirv& Spirv,
 	const TMap<FString, FVulkanShaderHeader::EType>& EntryTypes,
 	const FShaderCompilerInput& ShaderInput,
 	FVulkanHlslccHeader& CCHeader,
@@ -962,6 +962,8 @@ static void ConvertToNEWHeader(FOLDVulkanCodeHeader& OLDHeader,
 	}
 
 	// Add the packed global UBs
+	const FString UBOGlobalsNameSpv = ANSI_TO_TCHAR(CrossCompiler::FShaderConductorContext::GetIdentifierTable().GlobalsUniformBuffer);
+
 	for (int32 Index = 0; Index < OLDHeader.NEWPackedUBToVulkanBindingIndices.Num(); ++Index)
 	{
 		const FOLDVulkanCodeHeader::FPackedUBToVulkanBindingIndex& PackedArrayInfo = OLDHeader.NEWPackedUBToVulkanBindingIndices[Index];
@@ -970,11 +972,11 @@ static void ConvertToNEWHeader(FOLDVulkanCodeHeader& OLDHeader,
 		PackedUB.PackedTypeIndex = CrossCompiler::PackedTypeNameToTypeIndex(PackedArrayInfo.TypeName);
 		PackedUB.SizeInBytes = OLDHeader.NEWPackedGlobalUBSizes[Index];
 
-		const FSpirv::FEntry* Entry = Spirv.GetEntryByBindingIndex(PackedArrayInfo.VulkanBindingIndex);
+		const FVulkanSpirv::FEntry* Entry = Spirv.GetEntryByBindingIndex(PackedArrayInfo.VulkanBindingIndex);
 		check(Entry);
 
 		// We are dealing with "HLSLCC_CB" for HLSLcc, and "$Globals" for DXC
-		check(Entry->Name.StartsWith(TEXT("HLSLCC_CB")) || Entry->Name.StartsWith(TEXT("$Globals")));
+		check(Entry->Name.StartsWith(TEXT("HLSLCC_CB")) || Entry->Name.StartsWith(UBOGlobalsNameSpv));
 
 		PackedUB.SPIRVDescriptorSetOffset = Entry->WordDescriptorSetIndex;
 		PackedUB.SPIRVBindingIndexOffset = Entry->WordBindingIndex;
@@ -1018,7 +1020,7 @@ static void ConvertToNEWHeader(FOLDVulkanCodeHeader& OLDHeader,
 				check(HeaderGlobalIndex != INDEX_NONE);
 				check(GlobalNames[HeaderGlobalIndex] == AttachmentName);
 				FVulkanShaderHeader::FGlobalInfo& GlobalInfo = OutHeader.Globals[HeaderGlobalIndex];
-				const FSpirv::FEntry* Entry = Spirv.GetEntry(AttachmentName);
+				const FVulkanSpirv::FEntry* Entry = Spirv.GetEntry(AttachmentName);
 				check(Entry);
 				check(Entry->Binding != -1);
 
@@ -1059,7 +1061,7 @@ static void BuildShaderOutput(
 	int32						SourceLen,
 	const FVulkanBindingTable&	BindingTable,
 	uint32						NumLines,
-	FSpirv&						Spirv,
+	FVulkanSpirv&				Spirv,
 	const FString&				DebugName,
 	bool						bHasRealUBs,
 	bool						bSourceContainsMetaDataOnly)
@@ -1525,7 +1527,7 @@ static void BuildShaderOutput(
 	{
 		if (IsVulkanMobilePlatform((EShaderPlatform)ShaderInput.Target.Platform))
 		{
-			CompileOfflineMali(ShaderInput, ShaderOutput, (const ANSICHAR*)Spirv.Data.GetData(), Spirv.Data.Num() * sizeof(uint32), true, (const ANSICHAR*)(Spirv.Data.GetData() + Spirv.OffsetToMainName));
+			CompileOfflineMali(ShaderInput, ShaderOutput, (const ANSICHAR*)Spirv.GetByteData(), Spirv.GetByteSize(), true, (const ANSICHAR*)(Spirv.GetByteData() + Spirv.OffsetToMainName));
 		}
 	}
 
@@ -1704,7 +1706,7 @@ FCompilerInfo::FCompilerInfo(const FShaderCompilerInput& InInput, const FString&
 static bool CompileUsingInternal(FCompilerInfo& CompilerInfo, const FVulkanBindingTable& BindingTable, const TArray<ANSICHAR>& GlslSource, FShaderCompilerOutput& Output, bool bHasRealUBs)
 {
 	FString Errors;
-	FSpirv Spirv;
+	FVulkanSpirv Spirv;
 	const ANSICHAR* Main = GlslSource.GetData();
 	Main = FCStringAnsi::Strstr(Main, "void main_");
 	check(Main);
@@ -1738,7 +1740,7 @@ static bool CompileUsingInternal(FCompilerInfo& CompilerInfo, const FVulkanBindi
 			FArchive* FileWriter = IFileManager::Get().CreateFileWriter(*InfoFile);
 			if (FileWriter)
 			{
-				FString OutputString = FString::Printf(TEXT("main_%0.8x_%0.8x\n"), Spirv.Data.Num() * sizeof(uint32), Spirv.CRC);
+				FString OutputString = FString::Printf(TEXT("main_%0.8x_%0.8x\n"), Spirv.GetByteSize(), Spirv.CRC);
 				auto AnsiOutputString = StringCast<ANSICHAR>(*OutputString);
 				FileWriter->Serialize((ANSICHAR*)AnsiOutputString.Get(), AnsiOutputString.Length());
 				FileWriter->Close();
@@ -1840,224 +1842,21 @@ static bool CompileWithHlslcc(const FString& PreprocessedShader, FVulkanBindingT
 
 #if PLATFORM_MAC || PLATFORM_WINDOWS || PLATFORM_LINUX
 
-// Container structure for all SPIR-V reflection resources and in/out attributes.
-struct FSpirvReflectionBindings
-{
-	TArray<SpvReflectInterfaceVariable*> InputAttributes;
-	TArray<SpvReflectInterfaceVariable*> OutputAttributes;
-	TSet<SpvReflectDescriptorBinding*> AtomicCounters;
-	TArray<SpvReflectDescriptorBinding*> InputAttachments; // for subpass inputs
-	TArray<SpvReflectDescriptorBinding*> UniformBuffers;
-	TArray<SpvReflectDescriptorBinding*> Samplers;
-	TArray<SpvReflectDescriptorBinding*> TextureSRVs;
-	TArray<SpvReflectDescriptorBinding*> TextureUAVs;
-	TArray<SpvReflectDescriptorBinding*> TBufferSRVs;
-	TArray<SpvReflectDescriptorBinding*> TBufferUAVs;
-	TArray<SpvReflectDescriptorBinding*> SBufferSRVs;
-	TArray<SpvReflectDescriptorBinding*> SBufferUAVs;
-};
-
-// Parse the index from a semantic name, e.g. "ATTRIBUTE14" returns 14.
-static bool ParseSemanticIndex(const ANSICHAR* InSemanticName, int32& OutSemanticIndex)
-{
-	if (!InSemanticName)
-	{
-		return false;
-	}
-
-	for (int32 NameLen = FCStringAnsi::Strlen(InSemanticName), Index = NameLen; Index > 0; --Index)
-	{
-		if (!FChar::IsDigit(InSemanticName[Index - 1]))
-		{
-			if (Index == NameLen)
-			{
-				// Semantic name does not end with digits
-				return false;
-			}
-			else
-			{
-				// Return suffix numeric starting at previous string position
-				OutSemanticIndex = FCStringAnsi::Atoi(InSemanticName + Index);
-				return true;
-			}
-		}
-	}
-
-	return false;
-}
-
-static void GatherSpirvReflectionBindingEntry(SpvReflectDescriptorBinding* InBinding, FSpirvReflectionBindings& OutBindings)
-{
-	switch (InBinding->resource_type)
-	{
-	case SPV_REFLECT_RESOURCE_FLAG_SAMPLER:
-	{
-		// Gather sampler states (i.e. SamplerState or SamplerComparisonState)
-		check(InBinding->descriptor_type == SPV_REFLECT_DESCRIPTOR_TYPE_SAMPLER);
-		if (InBinding->accessed)
-		{
-			OutBindings.Samplers.Add(InBinding);
-		}
-		break;
-	}
-	case SPV_REFLECT_RESOURCE_FLAG_CBV:
-	{
-		// Gather constant buffers (i.e. cbuffer or ConstantBuffer<T>).
-		check(InBinding->descriptor_type == SPV_REFLECT_DESCRIPTOR_TYPE_UNIFORM_BUFFER);
-		if (InBinding->accessed)
-		{
-			OutBindings.UniformBuffers.Add(InBinding);
-		}
-		break;
-	}
-	case SPV_REFLECT_RESOURCE_FLAG_SRV:
-	{
-		// Gather SRV resources (e.g. Buffer, StructuredBuffer, Texture2D etc.)
-		switch (InBinding->descriptor_type)
-		{
-			/*case SPV_REFLECT_DESCRIPTOR_TYPE_COMBINED_IMAGE_SAMPLER:
-			{
-				if (InBinding->accessed)
-				{
-					OutBindings.Samplers.Add(InBinding);
-				}
-				break;
-			}*/
-		case SPV_REFLECT_DESCRIPTOR_TYPE_SAMPLED_IMAGE:
-		{
-			if (InBinding->accessed)
-			{
-				OutBindings.TextureSRVs.Add(InBinding);
-			}
-			break;
-		}
-		case SPV_REFLECT_DESCRIPTOR_TYPE_UNIFORM_TEXEL_BUFFER:
-		{
-			if (InBinding->accessed)
-			{
-				OutBindings.TBufferSRVs.Add(InBinding);
-			}
-			break;
-		}
-		case SPV_REFLECT_DESCRIPTOR_TYPE_STORAGE_BUFFER:
-		{
-			if (InBinding->accessed)
-			{
-				// Storage buffers must always occupy a UAV binding slot
-				OutBindings.SBufferUAVs.Add(InBinding);
-			}
-			break;
-		}
-		default:
-		{
-			// check(false);
-			break;
-		}
-		}
-		break;
-	}
-	case SPV_REFLECT_RESOURCE_FLAG_UAV:
-	{
-		if (InBinding->uav_counter_binding)
-		{
-			OutBindings.AtomicCounters.Add(InBinding->uav_counter_binding);
-		}
-
-		switch (InBinding->descriptor_type)
-		{
-		case SPV_REFLECT_DESCRIPTOR_TYPE_STORAGE_IMAGE:
-		{
-			OutBindings.TextureUAVs.Add(InBinding);
-			break;
-		}
-		case SPV_REFLECT_DESCRIPTOR_TYPE_STORAGE_TEXEL_BUFFER:
-		{
-			OutBindings.TBufferUAVs.Add(InBinding);
-			break;
-		}
-		case SPV_REFLECT_DESCRIPTOR_TYPE_STORAGE_BUFFER:
-		{
-			if (!OutBindings.AtomicCounters.Contains(InBinding) || InBinding->accessed)
-			{
-				OutBindings.SBufferUAVs.Add(InBinding);
-			}
-			break;
-		}
-		default:
-		{
-			break;
-		}
-		}
-		break;
-	}
-	default:
-	{
-		// Gather input attachments (e.g. subpass inputs)
-		switch (InBinding->descriptor_type)
-		{
-		case SPV_REFLECT_DESCRIPTOR_TYPE_INPUT_ATTACHMENT:
-		{
-			if (InBinding->accessed)
-			{
-				OutBindings.InputAttachments.Add(InBinding);
-			}
-			break;
-		}
-		default:
-		{
-			break;
-		}
-		}
-		break;
-	}
-	} // switch
-}
-
 static void GatherSpirvReflectionBindings(
 	spv_reflect::ShaderModule&	Reflection,
-	FSpirvReflectionBindings&	OutBindings,
+	FSpirvReflectBindings&		OutBindings,
 	const EShaderFrequency		ShaderFrequency)
 {
 	SpvReflectResult SpvResult = SPV_REFLECT_RESULT_NOT_READY;
 
-	// Enumerate all input attributes
-	uint32 NumInputAttributes = 0;
-	SpvResult = Reflection.EnumerateEntryPointInputVariables(Reflection.GetEntryPointName(), &NumInputAttributes, nullptr);
-	check(SpvResult == SPV_REFLECT_RESULT_SUCCESS);
-
-	if (NumInputAttributes > 0)
-	{
-		OutBindings.InputAttributes.SetNum(NumInputAttributes);
-		SpvResult = Reflection.EnumerateEntryPointInputVariables(Reflection.GetEntryPointName(), &NumInputAttributes, OutBindings.InputAttributes.GetData());
-		check(SpvResult == SPV_REFLECT_RESULT_SUCCESS);
-	}
+	OutBindings.GatherInputAttributes(Reflection);
+	OutBindings.GatherOutputAttributes(Reflection);
+	OutBindings.GatherDescriptorBindings(Reflection);
 
 	// Change indices of input attributes by their name suffix. Only in the vertex shader stage, "ATTRIBUTE" semantics have a special meaning for shader attributes.
 	if (ShaderFrequency == SF_Vertex)
 	{
-		for (SpvReflectInterfaceVariable* InterfaceVar : OutBindings.InputAttributes)
-		{
-			if (InterfaceVar->built_in == -1 && InterfaceVar->name && FCStringAnsi::Strncmp(InterfaceVar->name, "in.var.ATTRIBUTE", 16) == 0)
-			{
-				int32 Location = 0;
-				if (ParseSemanticIndex(InterfaceVar->name, Location))
-				{
-					Reflection.ChangeInputVariableLocation(InterfaceVar, static_cast<uint32>(Location));
-				}
-			}
-		}
-	}
-
-	// Enumerate all output attributes
-	uint32 NumOutputAttributes = 0;
-	SpvResult = Reflection.EnumerateEntryPointOutputVariables(Reflection.GetEntryPointName(), &NumOutputAttributes, nullptr);
-	check(SpvResult == SPV_REFLECT_RESULT_SUCCESS);
-
-	if (NumOutputAttributes > 0)
-	{
-		OutBindings.OutputAttributes.SetNum(NumOutputAttributes);
-		SpvResult = Reflection.EnumerateEntryPointOutputVariables(Reflection.GetEntryPointName(), &NumOutputAttributes, OutBindings.OutputAttributes.GetData());
-		check(SpvResult == SPV_REFLECT_RESULT_SUCCESS);
+		OutBindings.AssignInputAttributeLocationsBySemanticIndex(Reflection, CrossCompiler::FShaderConductorContext::GetIdentifierTable().InputAttribute);
 	}
 
 	// Change descriptor set numbers
@@ -2078,30 +1877,10 @@ static void GatherSpirvReflectionBindings(
 			Reflection.ChangeDescriptorSetNumber(DescSet, DescSetNo);
 		}
 	}
-
-	// Enumerate all binding points
-	TArray<SpvReflectDescriptorBinding*> Bindings;
-	uint32 NumBindings = 0;
-
-	SpvResult = Reflection.EnumerateDescriptorBindings(&NumBindings, nullptr);
-	check(SpvResult == SPV_REFLECT_RESULT_SUCCESS);
-	if (NumBindings > 0)
-	{
-		Bindings.SetNum(NumBindings);
-		SpvResult = Reflection.EnumerateDescriptorBindings(&NumBindings, Bindings.GetData());
-		check(SpvResult == SPV_REFLECT_RESULT_SUCCESS);
-
-		// Extract all the bindings first so that we process them in order - this lets us assign UAVs before other resources
-		// Which is necessary to match the D3D binding scheme.
-		for (auto const& BindingEntry : Bindings)
-		{
-			GatherSpirvReflectionBindingEntry(BindingEntry, OutBindings);
-		}
-	}
 }
 
 static void BuildShaderOutputFromSpirv(
-	FSpirv&						Spirv,
+	FVulkanSpirv&				Spirv,
 	const FShaderCompilerInput& Input,
 	FShaderCompilerOutput&		Output,
 	FVulkanBindingTable&		BindingTable,
@@ -2120,8 +1899,8 @@ static void BuildShaderOutputFromSpirv(
 	TMap<const SpvReflectDescriptorBinding*, uint32> SamplerStatesUseCount;
 
 	// Reflect SPIR-V module with SPIRV-Reflect library
-	const size_t SpirvDataSize = (Spirv.Data.Num() * sizeof(uint32));
-	spv_reflect::ShaderModule Reflection(SpirvDataSize, Spirv.Data.GetData(), SPV_REFLECT_RETURN_FLAG_SAMPLER_IMAGE_USAGE);
+	const size_t SpirvDataSize = Spirv.GetByteSize();
+	spv_reflect::ShaderModule Reflection(SpirvDataSize, Spirv.GetByteData(), SPV_REFLECT_RETURN_FLAG_SAMPLER_IMAGE_USAGE);
 	check(Reflection.GetResult() == SPV_REFLECT_RESULT_SUCCESS);
 
 	// Change final entry point name in SPIR-V module
@@ -2132,7 +1911,7 @@ static void BuildShaderOutputFromSpirv(
 		check(Result == SPV_REFLECT_RESULT_SUCCESS);
 	}
 
-	FSpirvReflectionBindings Bindings;
+	FSpirvReflectBindings Bindings;
 	GatherSpirvReflectionBindings(Reflection, Bindings, static_cast<EShaderFrequency>(Input.Target.Frequency));
 
 	// Register how often a sampler-state is used
@@ -2157,8 +1936,7 @@ static void BuildShaderOutputFromSpirv(
 		BindingTable.RegisterBinding(Attribute->semantic, "a", EVulkanBindingType::InputAttachment);
 	}*/
 
-	const FString UBOGlobalsNameSpv = TEXT("$Globals");
-	const FString UBOGlobalsNameGLSL = TEXT("_Globals");
+	const FString UBOGlobalsNameSpv = ANSI_TO_TCHAR(CrossCompiler::FShaderConductorContext::GetIdentifierTable().GlobalsUniformBuffer);
 
 	for (const SpvReflectDescriptorBinding* Binding : Bindings.UniformBuffers)
 	{
@@ -2248,7 +2026,7 @@ static void BuildShaderOutputFromSpirv(
 
 	int32 UBOBindings = 0, UAVBindings = 0, SRVBindings = 0, SMPBindings = 0, GLOBindings = 0;
 
-	auto GetRealBindingIndex = [&BindingTable, &BindingToIndexMap](const SpvReflectDescriptorBinding* Binding)
+	auto MapDescriptorBindingToIndex = [&BindingTable, &BindingToIndexMap](const SpvReflectDescriptorBinding* Binding)
 	{
 		return BindingTable.GetRealBindingIndex(BindingToIndexMap[Binding]);
 	};
@@ -2260,12 +2038,12 @@ static void BuildShaderOutputFromSpirv(
 		if (ResourceName == UBOGlobalsNameSpv)
 		{
 			// Register binding for uniform buffer
-			int32 BindingIndex = GetRealBindingIndex(Binding);
+			int32 BindingIndex = MapDescriptorBindingToIndex(Binding);
 
 			SpvResult = Reflection.ChangeDescriptorBindingNumbers(Binding, BindingIndex);// , GlobalSetId);
 			check(SpvResult == SPV_REFLECT_RESULT_SUCCESS);
 
-			Spirv.ReflectionInfo.Add(FSpirv::FEntry(UBOGlobalsNameSpv, BindingIndex));
+			Spirv.ReflectionInfo.Add(FVulkanSpirv::FEntry(UBOGlobalsNameSpv, BindingIndex));
 
 			CCHeaderWriter.WriteUniformBlock(TEXT("_Globals_h"), UBOBindings++);
 
@@ -2291,11 +2069,11 @@ static void BuildShaderOutputFromSpirv(
 		{
 			// Register uniform buffer
 
-			int32 BindingIndex = GetRealBindingIndex(Binding);
+			int32 BindingIndex = MapDescriptorBindingToIndex(Binding);
 			SpvResult = Reflection.ChangeDescriptorBindingNumbers(Binding, BindingIndex);//, GlobalSetId);
 			check(SpvResult == SPV_REFLECT_RESULT_SUCCESS);
 
-			Spirv.ReflectionInfo.Add(FSpirv::FEntry(ResourceName, BindingIndex));
+			Spirv.ReflectionInfo.Add(FVulkanSpirv::FEntry(ResourceName, BindingIndex));
 
 			CCHeaderWriter.WriteUniformBlock(*ResourceName, UBOBindings++);
 		}
@@ -2303,19 +2081,19 @@ static void BuildShaderOutputFromSpirv(
 
 	for (const SpvReflectDescriptorBinding* Binding : Bindings.InputAttachments)
 	{
-		int32 BindingIndex = GetRealBindingIndex(Binding);
+		int32 BindingIndex = MapDescriptorBindingToIndex(Binding);
 
 		SpvResult = Reflection.ChangeDescriptorBindingNumbers(Binding, BindingIndex);//, GlobalSetId);
 		check(SpvResult == SPV_REFLECT_RESULT_SUCCESS);
 
 		const FString ResourceName(ANSI_TO_TCHAR(Binding->name));
 
-		Spirv.ReflectionInfo.Add(FSpirv::FEntry(ResourceName, BindingIndex));
+		Spirv.ReflectionInfo.Add(FVulkanSpirv::FEntry(ResourceName, BindingIndex));
 	}
 
 	for (const SpvReflectDescriptorBinding* Binding : Bindings.TBufferUAVs)
 	{
-		int32 BindingIndex = GetRealBindingIndex(Binding);
+		int32 BindingIndex = MapDescriptorBindingToIndex(Binding);
 
 		SpvResult = Reflection.ChangeDescriptorBindingNumbers(Binding, BindingIndex);//, GlobalSetId);
 		check(SpvResult == SPV_REFLECT_RESULT_SUCCESS);
@@ -2323,12 +2101,12 @@ static void BuildShaderOutputFromSpirv(
 		const FString ResourceName(ANSI_TO_TCHAR(Binding->name));
 		CCHeaderWriter.WriteUAV(*ResourceName, UAVBindings++);
 
-		Spirv.ReflectionInfo.Add(FSpirv::FEntry(ResourceName, BindingIndex));
+		Spirv.ReflectionInfo.Add(FVulkanSpirv::FEntry(ResourceName, BindingIndex));
 	}
 
 	for (const SpvReflectDescriptorBinding* Binding : Bindings.SBufferUAVs)
 	{
-		int32 BindingIndex = GetRealBindingIndex(Binding);
+		int32 BindingIndex = MapDescriptorBindingToIndex(Binding);
 
 		SpvResult = Reflection.ChangeDescriptorBindingNumbers(Binding, BindingIndex);//, GlobalSetId);
 		check(SpvResult == SPV_REFLECT_RESULT_SUCCESS);
@@ -2336,12 +2114,12 @@ static void BuildShaderOutputFromSpirv(
 		const FString ResourceName(ANSI_TO_TCHAR(Binding->name));
 		CCHeaderWriter.WriteUAV(*ResourceName, UAVBindings++);
 
-		Spirv.ReflectionInfo.Add(FSpirv::FEntry(ResourceName, BindingIndex));
+		Spirv.ReflectionInfo.Add(FVulkanSpirv::FEntry(ResourceName, BindingIndex));
 	}
 
 	for (const SpvReflectDescriptorBinding* Binding : Bindings.TextureUAVs)
 	{
-		int32 BindingIndex = GetRealBindingIndex(Binding);
+		int32 BindingIndex = MapDescriptorBindingToIndex(Binding);
 
 		SpvResult = Reflection.ChangeDescriptorBindingNumbers(Binding, BindingIndex);//, GlobalSetId);
 		check(SpvResult == SPV_REFLECT_RESULT_SUCCESS);
@@ -2349,12 +2127,12 @@ static void BuildShaderOutputFromSpirv(
 		const FString ResourceName(ANSI_TO_TCHAR(Binding->name));
 		CCHeaderWriter.WriteUAV(*ResourceName, UAVBindings++);
 
-		Spirv.ReflectionInfo.Add(FSpirv::FEntry(ResourceName, BindingIndex));
+		Spirv.ReflectionInfo.Add(FVulkanSpirv::FEntry(ResourceName, BindingIndex));
 	}
 
 	for (const SpvReflectDescriptorBinding* Binding : Bindings.TBufferSRVs)
 	{
-		int32 BindingIndex = GetRealBindingIndex(Binding);
+		int32 BindingIndex = MapDescriptorBindingToIndex(Binding);
 
 		SpvResult = Reflection.ChangeDescriptorBindingNumbers(Binding, BindingIndex);//, GlobalSetId);
 		check(SpvResult == SPV_REFLECT_RESULT_SUCCESS);
@@ -2362,12 +2140,12 @@ static void BuildShaderOutputFromSpirv(
 		const FString ResourceName(ANSI_TO_TCHAR(Binding->name));
 		CCHeaderWriter.WriteSRV(*ResourceName, UAVBindings++);
 
-		Spirv.ReflectionInfo.Add(FSpirv::FEntry(ResourceName, BindingIndex));
+		Spirv.ReflectionInfo.Add(FVulkanSpirv::FEntry(ResourceName, BindingIndex));
 	}
 
 	for (const SpvReflectDescriptorBinding* Binding : Bindings.SBufferSRVs)
 	{
-		int32 BindingIndex = GetRealBindingIndex(Binding);
+		int32 BindingIndex = MapDescriptorBindingToIndex(Binding);
 
 		SpvResult = Reflection.ChangeDescriptorBindingNumbers(Binding, BindingIndex);//, GlobalSetId);
 		check(SpvResult == SPV_REFLECT_RESULT_SUCCESS);
@@ -2375,12 +2153,12 @@ static void BuildShaderOutputFromSpirv(
 		const FString ResourceName(ANSI_TO_TCHAR(Binding->name));
 		CCHeaderWriter.WriteSRV(*ResourceName, SRVBindings++);
 
-		Spirv.ReflectionInfo.Add(FSpirv::FEntry(ResourceName, BindingIndex));
+		Spirv.ReflectionInfo.Add(FVulkanSpirv::FEntry(ResourceName, BindingIndex));
 	}
 
 	for (const SpvReflectDescriptorBinding* Binding : Bindings.TextureSRVs)
 	{
-		int32 BindingIndex = GetRealBindingIndex(Binding);
+		int32 BindingIndex = MapDescriptorBindingToIndex(Binding);
 
 		SpvResult = Reflection.ChangeDescriptorBindingNumbers(Binding, BindingIndex);//, GlobalSetId);
 		check(SpvResult == SPV_REFLECT_RESULT_SUCCESS);
@@ -2404,18 +2182,18 @@ static void BuildShaderOutputFromSpirv(
 			CCHeaderWriter.WriteSRV(*ResourceName, SRVBindings++);
 		}
 
-		Spirv.ReflectionInfo.Add(FSpirv::FEntry(ResourceName, BindingIndex));
+		Spirv.ReflectionInfo.Add(FVulkanSpirv::FEntry(ResourceName, BindingIndex));
 	}
 
 	for (const SpvReflectDescriptorBinding* Binding : Bindings.Samplers)
 	{
-		int32 BindingIndex = GetRealBindingIndex(Binding);
+		int32 BindingIndex = MapDescriptorBindingToIndex(Binding);
 
 		SpvResult = Reflection.ChangeDescriptorBindingNumbers(Binding, BindingIndex);//, GlobalSetId);
 		check(SpvResult == SPV_REFLECT_RESULT_SUCCESS);
 
 		const FString ResourceName(ANSI_TO_TCHAR(Binding->name));
-		Spirv.ReflectionInfo.Add(FSpirv::FEntry(ResourceName, BindingIndex));
+		Spirv.ReflectionInfo.Add(FVulkanSpirv::FEntry(ResourceName, BindingIndex));
 
 		// Only emit sampler state when its shared, i.e. used with at least 2 textures
 //		if (const uint32* UseCount = SamplerStatesUseCount.Find(Binding))
@@ -2459,7 +2237,7 @@ static void BuildShaderOutputFromSpirv(
 		DumpDebugShaderText(Input, MetaData, TEXT("meta.txt"));
 
 		// SPIR-V file (Binary)
-		DumpDebugShaderBinary(Input, Spirv.Data.GetData(), Spirv.Data.Num() * sizeof(uint32), TEXT("spv"));
+		DumpDebugShaderBinary(Input, Spirv.GetByteData(), Spirv.GetByteSize(), TEXT("spv"));
 
 		//@todo-lh: move this also to ShaderCompilerCommon module
 		// Disassembled SPIR-V file (Text)
@@ -2471,7 +2249,7 @@ static void BuildShaderOutputFromSpirv(
 			// Convert to STL container for glslang library
 			std::vector<uint32> SpirvData;
 			SpirvData.resize(Spirv.Data.Num());
-			FMemory::Memcpy(&SpirvData[0], Spirv.Data.GetData(), Spirv.Data.Num() * sizeof(uint32));
+			FMemory::Memcpy(&SpirvData[0], Spirv.GetByteData(), Spirv.GetByteSize());
 
 			// Emit disassembled SPIR-V
 			spv::Parameterize();
@@ -2534,7 +2312,7 @@ static bool CompileWithShaderConductor(
 	}
 
 	// Compile HLSL source to SPIR-V binary
-	FSpirv Spirv;
+	FVulkanSpirv Spirv;
 	if (!CompilerContext.CompileHlslToSpirv(Options, Spirv.Data))
 	{
 		CompilerContext.FlushErrors(Output.Errors);
@@ -2546,7 +2324,7 @@ static bool CompileWithShaderConductor(
 
 	// Write final output shader
 	Output.Target = Input.Target;
-	Output.ShaderCode.GetWriteAccess().Append(reinterpret_cast<const uint8*>(Spirv.Data.GetData()), Spirv.Data.Num() * sizeof(uint32));
+	Output.ShaderCode.GetWriteAccess().Append(reinterpret_cast<const uint8*>(Spirv.GetByteData()), Spirv.GetByteSize());
 
 	if (Input.Environment.CompilerFlags.Contains(CFLAG_KeepDebugInfo))
 	{
@@ -2557,7 +2335,7 @@ static bool CompileWithShaderConductor(
 
 	if (bDebugDump)
 	{
-		DumpDebugShaderBinary(Input, Spirv.Data.GetData(), Spirv.Data.Num() * sizeof(uint32), TEXT("spv"));
+		DumpDebugShaderBinary(Input, Spirv.GetByteData(), Spirv.GetByteSize(), TEXT("spv"));
 	}
 
 	// Flush warnings
diff --git a/Engine/Source/Developer/VulkanShaderFormat/Private/VulkanShaderFormat.h b/Engine/Source/Developer/VulkanShaderFormat/Private/VulkanShaderFormat.h
--- a/Engine/Source/Developer/VulkanShaderFormat/Private/VulkanShaderFormat.h
+++ b/Engine/Source/Developer/VulkanShaderFormat/Private/VulkanShaderFormat.h
@@ -3,6 +3,7 @@
 #pragma once
 
 #include "CoreMinimal.h"
+#include "SpirvCommon.h"
 #include "hlslcc.h"
 
 enum class EVulkanShaderVersion
@@ -30,9 +31,8 @@ struct FCompilerInfo
 };
 
 
-struct FSpirv
+struct FVulkanSpirv : FSpirv
 {
-	TArray<uint32> Data;
 	struct FEntry
 	{
 		FEntry() = default;
@@ -135,7 +135,7 @@ struct FSpirv
 };
 
 // Updates all reflection entries in the specified SPIR-V module.
-extern bool PatchSpirvReflectionEntriesAndEntryPoint(FSpirv& OutSpirv);
+extern bool PatchSpirvReflectionEntriesAndEntryPoint(FVulkanSpirv& OutSpirv);
 
 // Generates SPIR-V out of the specified GLSL source code.
-extern bool GenerateSpirv(const ANSICHAR* Source, FCompilerInfo& CompilerInfo, FString& OutErrors, const FString& DumpDebugInfoPath, FSpirv& OutSpirv);
+extern bool GenerateSpirv(const ANSICHAR* Source, FCompilerInfo& CompilerInfo, FString& OutErrors, const FString& DumpDebugInfoPath, FVulkanSpirv& OutSpirv);
