Fix binding reflection for Metal backend: only Vulkan needs storage buffers to always occupy a UAV binding slot.

#rnx

diff --git a/Engine/Source/Developer/ShaderCompilerCommon/Private/SpirvReflectCommon.cpp b/Engine/Source/Developer/ShaderCompilerCommon/Private/SpirvReflectCommon.cpp
--- a/Engine/Source/Developer/ShaderCompilerCommon/Private/SpirvReflectCommon.cpp
+++ b/Engine/Source/Developer/ShaderCompilerCommon/Private/SpirvReflectCommon.cpp
@@ -54,8 +54,7 @@ void FSpirvReflectBindings::AddDescriptorBinding(SpvReflectDescriptorBinding* In
 		case SPV_REFLECT_DESCRIPTOR_TYPE_STORAGE_BUFFER:
 			if (InBinding->accessed)
 			{
-				// Storage buffers must always occupy a UAV binding slot
-				SBufferUAVs.Add(InBinding);
+				SBufferSRVs.Add(InBinding);
 			}
 			break;
 
diff --git a/Engine/Source/Developer/VulkanShaderFormat/Private/VulkanShaderCompiler.cpp b/Engine/Source/Developer/VulkanShaderFormat/Private/VulkanShaderCompiler.cpp
--- a/Engine/Source/Developer/VulkanShaderFormat/Private/VulkanShaderCompiler.cpp
+++ b/Engine/Source/Developer/VulkanShaderFormat/Private/VulkanShaderCompiler.cpp
@@ -1849,6 +1849,10 @@ static void GatherSpirvReflectionBindings(
 	OutBindings.GatherOutputAttributes(Reflection);
 	OutBindings.GatherDescriptorBindings(Reflection);
 
+	// Storage buffers always occupy a UAV binding slot, so move all SBufferSRVs into the SBufferUAVs array
+	OutBindings.SBufferUAVs.Append(OutBindings.SBufferSRVs);
+	OutBindings.SBufferSRVs.Empty();
+
 	// Change indices of input attributes by their name suffix. Only in the vertex shader stage, "ATTRIBUTE" semantics have a special meaning for shader attributes.
 	if (ShaderFrequency == SF_Vertex)
 	{
