Fix Vulkan backend: SPIRV module was duplicated in output archive.

#rnx

diff --git a/Engine/Source/Developer/VulkanShaderFormat/Private/VulkanShaderCompiler.cpp b/Engine/Source/Developer/VulkanShaderFormat/Private/VulkanShaderCompiler.cpp
--- a/Engine/Source/Developer/VulkanShaderFormat/Private/VulkanShaderCompiler.cpp
+++ b/Engine/Source/Developer/VulkanShaderFormat/Private/VulkanShaderCompiler.cpp
@@ -2301,11 +2301,7 @@ static bool CompileWithShaderConductor(
 	}
 
 	// Build shader output and binding table
-	BuildShaderOutputFromSpirv(Spirv, CompilerInfo.Input, Output, BindingTable, bHasRealUBs, bDebugDump);
-
-	// Write final output shader
-	Output.Target = Input.Target;
-	Output.ShaderCode.GetWriteAccess().Append(reinterpret_cast<const uint8*>(Spirv.GetByteData()), Spirv.GetByteSize());
+	BuildShaderOutputFromSpirv(Spirv, Input, Output, BindingTable, bHasRealUBs, bDebugDump);
 
 	if (Input.Environment.CompilerFlags.Contains(CFLAG_KeepDebugInfo))
 	{
