Enabling DXC on Vulkan Mobile

diff --git a/Engine/Source/Developer/VulkanShaderFormat/Private/VulkanShaderCompiler.cpp b/Engine/Source/Developer/VulkanShaderFormat/Private/VulkanShaderCompiler.cpp
--- a/Engine/Source/Developer/VulkanShaderFormat/Private/VulkanShaderCompiler.cpp
+++ b/Engine/Source/Developer/VulkanShaderFormat/Private/VulkanShaderCompiler.cpp
@@ -581,22 +581,11 @@ static void AddUniformBuffer(FOLDVulkanCodeHeader& OLDHeader,
 	FShaderParameterMap& InOutParameterMap,
 	FVulkanShaderHeader& OutHeader,
 	TArray<FPatchType>& OutTypePatch,
-	TArray<FString>& GlobalNames,
-	bool bUseRealUBs)
+	TArray<FString>& GlobalNames)
 {
 	int32 HeaderUBIndex = -1;
-	if (bUseRealUBs)
-	{
-		HeaderUBIndex = OutHeader.UniformBuffers.AddZeroed();
-	}
-	else
-	{
-		if (BindingIndex >= OutHeader.UniformBuffers.Num())
-		{
-			OutHeader.UniformBuffers.SetNumZeroed(BindingIndex + 1);
-		}
-		HeaderUBIndex = BindingIndex;
-	}
+	HeaderUBIndex = OutHeader.UniformBuffers.AddZeroed();
+
 	FVulkanShaderHeader::FUniformBufferInfo& UBInfo = OutHeader.UniformBuffers[HeaderUBIndex];
 	const FUniformBufferEntry* UniformBufferEntry = ShaderInput.Environment.UniformBufferMap.Find(UBName);
 	UBInfo.LayoutHash = UniformBufferEntry ? UniformBufferEntry->LayoutHash : 0;
@@ -606,26 +595,19 @@ static void AddUniformBuffer(FOLDVulkanCodeHeader& OLDHeader,
 	const FVulkanSpirv::FEntry* Entry = Spirv.GetEntry(UBName);
 	if (Entry)
 	{
-		checkf(bUseRealUBs, TEXT("Internal error: Emulated should NOT have a binding!"));
 		UBInfo.bOnlyHasResources = false;
 		UBInfo.ConstantDataOriginalBindingIndex = BindingIndex;
-		if (bUseRealUBs)
-		{
-			// Only real UBs require an entry for SPIRV
-			int32 SpirvInfoIndex = OutHeader.UniformBufferSpirvInfos.Add(FVulkanShaderHeader::FSpirvInfo(Entry->WordDescriptorSetIndex, Entry->WordBindingIndex));
-			check(SpirvInfoIndex == HeaderUBIndex);
-		}
+		
+		int32 SpirvInfoIndex = OutHeader.UniformBufferSpirvInfos.Add(FVulkanShaderHeader::FSpirvInfo(Entry->WordDescriptorSetIndex, Entry->WordBindingIndex));
+		check(SpirvInfoIndex == HeaderUBIndex);
 	}
 	else
 	{
 		UBInfo.bOnlyHasResources = true;
 		UBInfo.ConstantDataOriginalBindingIndex = UINT16_MAX;
-		if (bUseRealUBs)
-		{
-			// Only real UBs require a dummy entry for SPIRV
-			int32 SpirvInfoIndex = OutHeader.UniformBufferSpirvInfos.Add(FVulkanShaderHeader::FSpirvInfo());
-			check(SpirvInfoIndex == HeaderUBIndex);
-		}
+
+		int32 SpirvInfoIndex = OutHeader.UniformBufferSpirvInfos.Add(FVulkanShaderHeader::FSpirvInfo());
+		check(SpirvInfoIndex == HeaderUBIndex);
 	}
 
 	// Add used resources...
@@ -641,7 +623,7 @@ static void AddUniformBuffer(FOLDVulkanCodeHeader& OLDHeader,
 	else
 	{
 		// If we're using real uniform buffers we have to have resources at least
-		check(!bUseRealUBs || !UBInfo.bOnlyHasResources);
+		check(!UBInfo.bOnlyHasResources);
 	}
 
 	// Currently we don't support mismatched uniform buffer layouts/cbuffers with resources!
@@ -749,7 +731,7 @@ static bool IsSamplerState(const CrossCompiler::FHlslccHeader& CCHeader, const F
 	return false;
 }
 
-static void PrepareGlobals(const FVulkanBindingTable& BindingTable, const FVulkanHlslccHeader& CCHeader, const FShaderCompilerResourceTable& SRT, const TMap<FString, FVulkanShaderHeader::EType>& EntryTypes, const FShaderCompilerInput& ShaderInput, const TArray<FString>& ParameterNames, FShaderParameterMap& ParameterMap, TArray<FString>& OutGlobalNames, FVulkanShaderHeader& OutHeader, bool bHasRealUBs)
+static void PrepareGlobals(const FVulkanBindingTable& BindingTable, const FVulkanHlslccHeader& CCHeader, const FShaderCompilerResourceTable& SRT, const TMap<FString, FVulkanShaderHeader::EType>& EntryTypes, const FShaderCompilerInput& ShaderInput, const TArray<FString>& ParameterNames, FShaderParameterMap& ParameterMap, TArray<FString>& OutGlobalNames, FVulkanShaderHeader& OutHeader)
 {
 	// First pass, gather names for all the Globals that are NOT Samplers
 	for (int32 ParameterIndex = 0; ParameterIndex < ParameterNames.Num(); ++ParameterIndex)
@@ -787,18 +769,10 @@ static void PrepareGlobals(const FVulkanBindingTable& BindingTable, const FVulka
 				}
 				break;
 			case FVulkanShaderHeader::UniformBuffer:
-				check(bHasRealUBs);
 				AddGlobalNamesForUB();
 				break;
 			case FVulkanShaderHeader::PackedGlobal:
-				if (bHasRealUBs || Size > 0)
-				{
-					// Ignore
-				}
-				else if (!bHasRealUBs)
-				{
-					AddGlobalNamesForUB();
-				}
+				// Ignore
 				break;
 			default:
 				check(0);
@@ -848,14 +822,6 @@ static void PrepareGlobals(const FVulkanBindingTable& BindingTable, const FVulka
 				AddGlobalNamesForUB();
 				break;
 			case FVulkanShaderHeader::PackedGlobal:
-				if (bHasRealUBs || Size > 0)
-				{
-					// Ignore
-				}
-				else if (!bHasRealUBs)
-				{
-					AddGlobalNamesForUB();
-				}
 				break;
 			default:
 				check(0);
@@ -889,8 +855,7 @@ static void ConvertToNEWHeader(FOLDVulkanCodeHeader& OLDHeader,
 	const FShaderCompilerInput& ShaderInput,
 	FVulkanHlslccHeader& CCHeader,
 	FShaderParameterMap& InOutParameterMap,
-	FVulkanShaderHeader& OutHeader,
-	bool bHasRealUBs)
+	FVulkanShaderHeader& OutHeader)
 {
 	// Names that match the Header.Globals array
 	TArray<FString> GlobalNames;
@@ -902,7 +867,7 @@ static void ConvertToNEWHeader(FOLDVulkanCodeHeader& OLDHeader,
 
 	const FShaderCompilerResourceTable& SRT = OLDHeader.SerializedBindings.ShaderResourceTable;
 
-	PrepareGlobals(BindingTable, CCHeader, SRT, EntryTypes, ShaderInput, ParameterNames, InOutParameterMap, GlobalNames, OutHeader, bHasRealUBs);
+	PrepareGlobals(BindingTable, CCHeader, SRT, EntryTypes, ShaderInput, ParameterNames, InOutParameterMap, GlobalNames, OutHeader);
 
 	for (int32 ParameterIndex = 0; ParameterIndex < ParameterNames.Num(); ++ParameterIndex)
 	{
@@ -931,36 +896,20 @@ static void ConvertToNEWHeader(FOLDVulkanCodeHeader& OLDHeader,
 				break;
 			case FVulkanShaderHeader::PackedGlobal:
 				{
-					if (bHasRealUBs || Size > 0)
-					{
-						int32 HeaderPackedGlobalIndex = OutHeader.PackedGlobals.AddZeroed();
-						FVulkanShaderHeader::FPackedGlobalInfo& PackedGlobalInfo = OutHeader.PackedGlobals[HeaderPackedGlobalIndex];
-						PackedGlobalInfo.PackedTypeIndex = CrossCompiler::PackedTypeNameToTypeIndex(OLDHeader.NEWPackedUBToVulkanBindingIndices[BufferIndex].TypeName);
-						PackedGlobalInfo.PackedUBIndex = BufferIndex;
-						checkf(Size > 0, TEXT("Assertion failed for shader parameter: %s"), *ParameterName);
-						PackedGlobalInfo.ConstantDataSizeInFloats = Size / sizeof(float);
+					int32 HeaderPackedGlobalIndex = OutHeader.PackedGlobals.AddZeroed();
+					FVulkanShaderHeader::FPackedGlobalInfo& PackedGlobalInfo = OutHeader.PackedGlobals[HeaderPackedGlobalIndex];
+					PackedGlobalInfo.PackedTypeIndex = CrossCompiler::PackedTypeNameToTypeIndex(OLDHeader.NEWPackedUBToVulkanBindingIndices[BufferIndex].TypeName);
+					PackedGlobalInfo.PackedUBIndex = BufferIndex;
+					checkf(Size > 0, TEXT("Assertion failed for shader parameter: %s"), *ParameterName);
+					PackedGlobalInfo.ConstantDataSizeInFloats = Size / sizeof(float);
 #if VULKAN_ENABLE_BINDING_DEBUG_NAMES
-						PackedGlobalInfo.DebugName = ParameterName;
+					PackedGlobalInfo.DebugName = ParameterName;
 #endif
-						// Keep the original parameter info from InOutParameterMap as it's a shortcut into the packed global array!
-					}
-					else if (!bHasRealUBs)
-					{
-						check(Size == 0);
-						check(EntryTypes.FindChecked(ParameterName) == FVulkanShaderHeader::PackedGlobal);
-						// Ignore, this is an empty param (Size == 0) for an emulated UB, but keep the original parameter info
-						// from InOutParameterMap as it's a shortcut into the packed global ub copies!
-						AddUniformBuffer(OLDHeader, BindingTable, ShaderInput, CCHeader, Spirv, ParameterName, BufferIndex, InOutParameterMap, OutHeader, TypePatchList, GlobalNames, bHasRealUBs);
-					}
-					else
-					{
-						check(0);
-					}
+					// Keep the original parameter info from InOutParameterMap as it's a shortcut into the packed global array!
 				}
 				break;
 			case FVulkanShaderHeader::UniformBuffer:
-				check(bHasRealUBs);
-				AddUniformBuffer(OLDHeader, BindingTable, ShaderInput, CCHeader, Spirv, ParameterName, BufferIndex, InOutParameterMap, OutHeader, TypePatchList, GlobalNames, bHasRealUBs);
+				AddUniformBuffer(OLDHeader, BindingTable, ShaderInput, CCHeader, Spirv, ParameterName, BufferIndex, InOutParameterMap, OutHeader, TypePatchList, GlobalNames);
 				break;
 			default:
 				check(0);
@@ -970,7 +919,7 @@ static void ConvertToNEWHeader(FOLDVulkanCodeHeader& OLDHeader,
 		else
 		{
 			// Not found means it's a new resource-only UniformBuffer
-			AddUniformBuffer(OLDHeader, BindingTable, ShaderInput, CCHeader, Spirv, ParameterName, BufferIndex, InOutParameterMap, OutHeader, TypePatchList, GlobalNames, bHasRealUBs);
+			AddUniformBuffer(OLDHeader, BindingTable, ShaderInput, CCHeader, Spirv, ParameterName, BufferIndex, InOutParameterMap, OutHeader, TypePatchList, GlobalNames);
 		}
 	}
 
@@ -1061,7 +1010,7 @@ static void ConvertToNEWHeader(FOLDVulkanCodeHeader& OLDHeader,
 		}
 	}
 
-	check(!bHasRealUBs || OLDHeader.UniformBuffersCopyInfo.Num() == 0);
+	check(OLDHeader.UniformBuffersCopyInfo.Num() == 0);
 	OutHeader.EmulatedUBsCopyInfo = OLDHeader.UniformBuffersCopyInfo;
 	OutHeader.EmulatedUBCopyRanges = OLDHeader.NEWEmulatedUBCopyRanges;
 	OutHeader.SourceHash = OLDHeader.SourceHash;
@@ -1070,7 +1019,6 @@ static void ConvertToNEWHeader(FOLDVulkanCodeHeader& OLDHeader,
 	OutHeader.DebugName = OLDHeader.ShaderName;
 #endif
 	OutHeader.InOutMask = OLDHeader.SerializedBindings.InOutMask.Bitmask;
-	OutHeader.bHasRealUBs = bHasRealUBs;
 }
 
 
@@ -1083,7 +1031,6 @@ static void BuildShaderOutput(
 	uint32						NumLines,
 	FVulkanSpirv&				Spirv,
 	const FString&				DebugName,
-	bool						bHasRealUBs,
 	bool						bSourceContainsMetaDataOnly)
 {
 	const ANSICHAR* USFSource = InShaderSource;
@@ -1309,7 +1256,6 @@ static void BuildShaderOutput(
 
 	for (auto& PackedUBCopy : CCHeader.PackedUBGlobalCopies)
 	{
-		ensure(!bHasRealUBs);
 		CrossCompiler::FUniformBufferCopyInfo CopyInfo;
 		CopyInfo.SourceUBIndex = PackedUBCopy.SourceUB;
 		CopyInfo.SourceOffsetInFloats = PackedUBCopy.SourceOffset;
@@ -1519,7 +1465,7 @@ static void BuildShaderOutput(
 	}
 
 	FVulkanShaderHeader NEWHeader(FVulkanShaderHeader::EZero);
-	ConvertToNEWHeader(OLDHeader, BindingTable, Spirv, NEWEntryTypes, ShaderInput, CCHeader, ShaderOutput.ParameterMap, NEWHeader, bHasRealUBs);
+	ConvertToNEWHeader(OLDHeader, BindingTable, Spirv, NEWEntryTypes, ShaderInput, CCHeader, ShaderOutput.ParameterMap, NEWHeader);
 
 	if (ShaderInput.Environment.CompilerFlags.Contains(CFLAG_KeepDebugInfo))
 	{
@@ -1724,7 +1670,7 @@ FCompilerInfo::FCompilerInfo(const FShaderCompilerInput& InInput, const FString&
 /**
  * Compile a shader using the internal shader compiling library
  */
-static bool CompileUsingInternal(FCompilerInfo& CompilerInfo, const FVulkanBindingTable& BindingTable, const TArray<ANSICHAR>& GlslSource, FShaderCompilerOutput& Output, bool bHasRealUBs)
+static bool CompileUsingInternal(FCompilerInfo& CompilerInfo, const FVulkanBindingTable& BindingTable, const TArray<ANSICHAR>& GlslSource, FShaderCompilerOutput& Output)
 {
 	FString Errors;
 	FVulkanSpirv Spirv;
@@ -1753,7 +1699,7 @@ static bool CompileUsingInternal(FCompilerInfo& CompilerInfo, const FVulkanBindi
 		Output.Target = CompilerInfo.Input.Target;
 		BuildShaderOutput(Output, CompilerInfo.Input,
 			GlslSource.GetData(), GlslSource.Num(),
-			BindingTable, NumLines, Spirv, DebugName, bHasRealUBs, false);
+			BindingTable, NumLines, Spirv, DebugName, false);
 
 		if (CompilerInfo.bDebugDump)
 		{
@@ -1911,7 +1857,6 @@ static bool BuildShaderOutputFromSpirv(
 	FShaderCompilerOutput&					Output,
 	FVulkanBindingTable&					BindingTable,
 	const FString&							EntryPointName,
-	bool									bHasRealUBs,
 	bool									bStripReflect,
 	bool									bIsRayTracingShader,
 	bool									bDebugDump)
@@ -2310,7 +2255,6 @@ static bool BuildShaderOutputFromSpirv(
 		/*NumLines*/0,
 		Spirv,
 		DebugName,
-		bHasRealUBs,
 		true // source contains meta data only
 	);
 
@@ -2332,7 +2276,6 @@ static bool CompileWithShaderConductor(
 	EHlslCompileTarget		HlslCompilerTarget,
 	FShaderCompilerOutput&	Output,
 	FVulkanBindingTable&	BindingTable,
-	bool					bHasRealUBs,
 	bool					bStripReflect)
 {
 	const FShaderCompilerInput& Input = CompilerInfo.Input;
@@ -2387,7 +2330,7 @@ static bool CompileWithShaderConductor(
 	}
 
 	// Build shader output and binding table
-	Output.bSucceeded = BuildShaderOutputFromSpirv(CompilerContext, Spirv, Input, Output, BindingTable, EntryPointName, bHasRealUBs, bStripReflect, bIsRayTracingShader, bDebugDump);
+	Output.bSucceeded = BuildShaderOutputFromSpirv(CompilerContext, Spirv, Input, Output, BindingTable, EntryPointName, bStripReflect, bIsRayTracingShader, bDebugDump);
 
 	if (Input.Environment.CompilerFlags.Contains(CFLAG_KeepDebugInfo))
 	{
@@ -2412,7 +2355,6 @@ void DoCompileVulkanShader(const FShaderCompilerInput& Input, FShaderCompilerOut
 	const EShaderPlatform ShaderPlatform = (EShaderPlatform)Input.Target.Platform;
 	check(IsVulkanPlatform(ShaderPlatform));
 
-	const bool bHasRealUBs = !Input.Environment.CompilerFlags.Contains(CFLAG_UseEmulatedUB);
 	const bool bIsSM5 = (Version == EVulkanShaderVersion::SM5);
 	const bool bIsMobile = (Version == EVulkanShaderVersion::ES3_1 || Version == EVulkanShaderVersion::ES3_1_ANDROID);
 	const bool bStripReflect = (IsVulkanMobilePlatform(ShaderPlatform) || IsVulkanMobileSM5Platform(ShaderPlatform));
@@ -2518,16 +2460,9 @@ void DoCompileVulkanShader(const FShaderCompilerInput& Input, FShaderCompilerOut
 	{
 		CompilerInfo.CCFlags |= HLSLCC_PackUniforms;
 		CompilerInfo.CCFlags |= HLSLCC_PackUniformsIntoUniformBuffers;
-		if (bHasRealUBs)
-		{
-			// Only flatten structures inside UBs
-			CompilerInfo.CCFlags |= HLSLCC_FlattenUniformBufferStructures;
-		}
-		else
-		{
-			// Flatten ALL UBs
-			CompilerInfo.CCFlags |= HLSLCC_FlattenUniformBuffers | HLSLCC_ExpandUBMemberArrays;
-		}
+
+		// Only flatten structures inside UBs
+		CompilerInfo.CCFlags |= HLSLCC_FlattenUniformBufferStructures;
 
 		if (bUseFullPrecisionInPS)
 		{
@@ -2566,7 +2501,7 @@ void DoCompileVulkanShader(const FShaderCompilerInput& Input, FShaderCompilerOut
 	if (bForceDXC)
 	{
 		// Cross-compile shader via ShaderConductor (DXC, SPIRV-Tools, SPIRV-Cross)
-		bSuccess = CompileWithShaderConductor(PreprocessedShaderSource, EntryPointName, CompilerInfo, HlslCompilerTarget, Output, BindingTable, bHasRealUBs, bStripReflect);
+		bSuccess = CompileWithShaderConductor(PreprocessedShaderSource, EntryPointName, CompilerInfo, HlslCompilerTarget, Output, BindingTable, bStripReflect);
 	}
 	else
 #endif // PLATFORM_MAC || PLATFORM_WINDOWS || PLATFORM_LINUX
@@ -2588,7 +2523,7 @@ void DoCompileVulkanShader(const FShaderCompilerInput& Input, FShaderCompilerOut
 			// For debugging: if you hit an error from Glslang/Spirv, use the SourceNoHeader for line numbers
 			auto* SourceWithHeader = GeneratedGlslSource.GetData();
 			char* SourceNoHeader = strstr(SourceWithHeader, "#version");
-			bSuccess = CompileUsingInternal(CompilerInfo, BindingTable, GeneratedGlslSource, Output, bHasRealUBs);
+			bSuccess = CompileUsingInternal(CompilerInfo, BindingTable, GeneratedGlslSource, Output);
 			if (bDirectCompile)
 			{
 				FPlatformMisc::LowLevelOutputDebugStringf(TEXT("Success: %d\n%s\n"), bSuccess, ANSI_TO_TCHAR(SourceWithHeader));
diff --git a/Engine/Source/Runtime/Engine/Private/ShaderCompiler/ShaderCompiler.cpp b/Engine/Source/Runtime/Engine/Private/ShaderCompiler/ShaderCompiler.cpp
--- a/Engine/Source/Runtime/Engine/Private/ShaderCompiler/ShaderCompiler.cpp
+++ b/Engine/Source/Runtime/Engine/Private/ShaderCompiler/ShaderCompiler.cpp
@@ -925,7 +925,7 @@ static TAutoConsoleVariable<int32> CVarOpenGLForceDXC(
 
 static TAutoConsoleVariable<int32> CVarVulkanForceDXC(
 	TEXT("r.Vulkan.ForceDXC"),
-	1,
+	3,
 	TEXT("Forces DirectX Shader Compiler (DXC) to be used for all Vulkan shaders instead of hlslcc.\n")
 	TEXT(" 0: Disable (hlslcc/glslang)\n")
 	TEXT(" 1: Enabled on desktop platforms only (default)\n")
@@ -5071,23 +5071,6 @@ void GlobalBeginCompileShader(
 				Input.Environment.CompilerFlags.Add(CFLAG_UseEmulatedUB);
 			}
 		}
-		else if(IsVulkanPlatform((EShaderPlatform)Target.Platform))
-		{
-			static const auto CVar = IConsoleManager::Get().FindConsoleVariable(TEXT("r.Vulkan.UseRealUBs"));
-			if ((CVar && CVar->GetInt() == 0) || 
-				Target.Platform == SP_VULKAN_ES3_1_ANDROID) // we force eUB on mobile Android
-			{
-				Input.Environment.CompilerFlags.Add(CFLAG_UseEmulatedUB);
-			}
-		}
-	}
-	else
-	{
-		static const auto CVar = IConsoleManager::Get().FindConsoleVariable(TEXT("r.Vulkan.UseRealUBs"));
-		if (CVar && CVar->GetInt() == 0)
-		{
-			Input.Environment.CompilerFlags.Add(CFLAG_UseEmulatedUB);
-		}
 	}
 
 	Input.Environment.SetDefine(TEXT("HAS_INVERTED_Z_BUFFER"), (bool)ERHIZBuffer::IsInverted);
diff --git a/Engine/Source/Runtime/RHI/Private/RHI.cpp b/Engine/Source/Runtime/RHI/Private/RHI.cpp
--- a/Engine/Source/Runtime/RHI/Private/RHI.cpp
+++ b/Engine/Source/Runtime/RHI/Private/RHI.cpp
@@ -54,14 +54,6 @@ IMPLEMENT_TYPE_LAYOUT(FRHIUniformBufferLayout::FResourceParameter);
 	#define RHIRESOURCE_NUM_FRAMES_TO_EXPIRE 3
 #endif
 
-static FAutoConsoleVariable CVarUseVulkanRealUBs(
-	TEXT("r.Vulkan.UseRealUBs"),
-	1,
-	TEXT("0: Emulate uniform buffers on Vulkan SM4/SM5 (debugging ONLY)\n")
-	TEXT("1: Use real uniform buffers [default]"),
-	ECVF_ReadOnly
-	);
-
 static TAutoConsoleVariable<int32> CVarDisableEngineAndAppRegistration(
 	TEXT("r.DisableEngineAndAppRegistration"),
 	0,
diff --git a/Engine/Source/Runtime/RenderCore/Private/Shader.cpp b/Engine/Source/Runtime/RenderCore/Private/Shader.cpp
--- a/Engine/Source/Runtime/RenderCore/Private/Shader.cpp
+++ b/Engine/Source/Runtime/RenderCore/Private/Shader.cpp
@@ -1493,16 +1493,6 @@ void ShaderMapAppendKeyString(EShaderPlatform Platform, FString& KeyString)
 			KeyString += (CVar && CVar->GetInt() != 0) ? TEXT("_NoUB") : TEXT("");
 		}
 
-		if (IsVulkanPlatform(Platform))
-		{
-			static IConsoleVariable* CVar = IConsoleManager::Get().FindConsoleVariable(TEXT("r.Vulkan.UseRealUBs"));
-			if ((CVar && CVar->GetInt() == 0) || 
-				Platform == SP_VULKAN_ES3_1_ANDROID) // we force eUB on mobile Android
-			{
-				KeyString += TEXT("_NoUB");
-			}
-		}
-
 		{
 			static FShaderPlatformCachedIniValue<bool> MobileEnableMovableSpotlightsIniValue(TEXT("/Script/Engine.RendererSettings"), TEXT("r.Mobile.EnableMovableSpotlights"));
 			static FShaderPlatformCachedIniValue<bool> MobileEnableMovableSpotlightsShadowIniValue(TEXT("/Script/Engine.RendererSettings"), TEXT("r.Mobile.EnableMovableSpotlightsShadow"));
@@ -1555,12 +1545,6 @@ void ShaderMapAppendKeyString(EShaderPlatform Platform, FString& KeyString)
 			static IConsoleVariable* CVar = IConsoleManager::Get().FindConsoleVariable(TEXT("OpenGL.UseEmulatedUBs"));
 			KeyString += (CVar && CVar->GetInt() != 0) ? TEXT("_NoUB") : TEXT("");
 		}
-
-		if (IsVulkanPlatform(Platform))
-		{
-			static IConsoleVariable* CVar = IConsoleManager::Get().FindConsoleVariable(TEXT("r.Vulkan.UseRealUBs"));
-			KeyString += (CVar && CVar->GetInt() == 0) ? TEXT("_NoUB") : TEXT("");
-		}
 	}
 
 	const IShaderFormat* ShaderFormat = GetTargetPlatformManagerRef().FindShaderFormat(ShaderFormatName);
diff --git a/Engine/Source/Runtime/RenderCore/Private/ShaderCodeLibrary.cpp b/Engine/Source/Runtime/RenderCore/Private/ShaderCodeLibrary.cpp
--- a/Engine/Source/Runtime/RenderCore/Private/ShaderCodeLibrary.cpp
+++ b/Engine/Source/Runtime/RenderCore/Private/ShaderCodeLibrary.cpp
@@ -2171,7 +2171,7 @@ public:
 			FName const& Format = Descriptor.ShaderFormat;
 
 			EShaderPlatform Platform = ShaderFormatToLegacyShaderPlatform(Format);
-			FName PossiblyAdjustedFormat = LegacyShaderPlatformToShaderFormat(Platform);	// Vulkan and GL switch between name variants depending on CVars (e.g. see r.Vulkan.UseRealUBs)
+			FName PossiblyAdjustedFormat = LegacyShaderPlatformToShaderFormat(Platform);	// Vulkan and GL switch between name variants depending on CVars 
 			FEditorShaderCodeArchive* CodeArchive = EditorShaderCodeArchive[Platform];
 			if (!CodeArchive)
 			{
@@ -2192,7 +2192,7 @@ public:
 			bool bUseStableKeys = Descriptor.bNeedsStableKeys;
 
 			EShaderPlatform Platform = ShaderFormatToLegacyShaderPlatform(Format);
-			FName PossiblyAdjustedFormat = LegacyShaderPlatformToShaderFormat(Platform);	// Vulkan and GL switch between name variants depending on CVars (e.g. see r.Vulkan.UseRealUBs)
+			FName PossiblyAdjustedFormat = LegacyShaderPlatformToShaderFormat(Platform);	// Vulkan and GL switch between name variants depending on CVars 
 			FEditorShaderStableInfo* StableArchive = EditorShaderStableInfo[Platform];
 			if (!StableArchive && bUseStableKeys)
 			{
diff --git a/Engine/Source/Runtime/VulkanRHI/Private/Android/VulkanAndroidPlatform.h b/Engine/Source/Runtime/VulkanRHI/Private/Android/VulkanAndroidPlatform.h
--- a/Engine/Source/Runtime/VulkanRHI/Private/Android/VulkanAndroidPlatform.h
+++ b/Engine/Source/Runtime/VulkanRHI/Private/Android/VulkanAndroidPlatform.h
@@ -107,11 +107,6 @@ public:
 
 	static bool RegisterGPUWork() { return false; }
 
-	static bool UseRealUBsOptimization(bool bCodeHeaderUseRealUBs)
-	{
-		return !RequiresMobileRenderer();
-	}
-
 	// Assume most devices can't use the extra cores for running parallel tasks
 	static bool SupportParallelRenderingTasks() { return false; }
 
diff --git a/Engine/Source/Runtime/VulkanRHI/Private/Linux/VulkanLinuxPlatform.h b/Engine/Source/Runtime/VulkanRHI/Private/Linux/VulkanLinuxPlatform.h
--- a/Engine/Source/Runtime/VulkanRHI/Private/Linux/VulkanLinuxPlatform.h
+++ b/Engine/Source/Runtime/VulkanRHI/Private/Linux/VulkanLinuxPlatform.h
@@ -49,25 +49,6 @@ public:
 
 	static void CreateSurface(void* WindowHandle, VkInstance Instance, VkSurfaceKHR* OutSurface);
 
-	// Some platforms only support real or non-real UBs, so this function can optimize it out
-	static bool UseRealUBsOptimization(bool bCodeHeaderUseRealUBs)
-	{
-		// cooked builds will return the bool unchanged - relying on the compiler to optimize out the editor code path
-		if (UE_EDITOR)
-		{
-			static bool bAlwaysUseRealUBs([]()
-			{
-				static auto* CVar = IConsoleManager::Get().FindTConsoleVariableDataInt(TEXT("r.Vulkan.UseRealUBs"));
-				return (CVar && CVar->GetValueOnAnyThread() == 0);
-			}());
-			return bAlwaysUseRealUBs ? false : bCodeHeaderUseRealUBs;
-		}
-		else
-		{
-			return bCodeHeaderUseRealUBs;
-		}
-	}
-
 	static bool ForceEnableDebugMarkers();
 
 	static void WriteCrashMarker(const FOptionalVulkanDeviceExtensions& OptionalExtensions, VkCommandBuffer CmdBuffer, VkBuffer DestBuffer, const TArrayView<uint32>& Entries, bool bAdding);
diff --git a/Engine/Source/Runtime/VulkanRHI/Private/VulkanCommands.cpp b/Engine/Source/Runtime/VulkanRHI/Private/VulkanCommands.cpp
--- a/Engine/Source/Runtime/VulkanRHI/Private/VulkanCommands.cpp
+++ b/Engine/Source/Runtime/VulkanRHI/Private/VulkanCommands.cpp
@@ -427,37 +427,28 @@ inline void FVulkanCommandListContext::SetShaderUniformBuffer(ShaderStage::EStag
 	check(Shader->GetShaderKey() == PendingGfxState->GetCurrentShaderKey(Stage));
 
 	const FVulkanShaderHeader& CodeHeader = Shader->GetCodeHeader();
-	const bool bUseRealUBs = FVulkanPlatform::UseRealUBsOptimization(CodeHeader.bHasRealUBs != 0);
 	const FVulkanShaderHeader::FUniformBufferInfo& HeaderUBInfo = CodeHeader.UniformBuffers[BufferIndex];
 	checkfSlow(!HeaderUBInfo.LayoutHash || HeaderUBInfo.LayoutHash == UniformBuffer->GetLayout().GetHash(), TEXT("Mismatched UB layout! Got hash 0x%x, expected 0x%x!"), UniformBuffer->GetLayout().GetHash(), HeaderUBInfo.LayoutHash);
 	const FVulkanGfxPipelineDescriptorInfo& DescriptorInfo = PendingGfxState->CurrentState->GetGfxPipelineDescriptorInfo();
-	if (!bUseRealUBs || !HeaderUBInfo.bOnlyHasResources)
+	if (!HeaderUBInfo.bOnlyHasResources)
 	{
-		checkSlow(!bUseRealUBs || UniformBuffer->GetLayout().ConstantBufferSize > 0);
+		checkSlow(UniformBuffer->GetLayout().ConstantBufferSize > 0);
 		extern TAutoConsoleVariable<int32> GDynamicGlobalUBs;
-		if (bUseRealUBs)
+
+		uint8 DescriptorSet;
+		uint32 BindingIndex;
+		if (!DescriptorInfo.GetDescriptorSetAndBindingIndex(FVulkanShaderHeader::UniformBuffer, Stage, BufferIndex, DescriptorSet, BindingIndex))
 		{
-			uint8 DescriptorSet;
-			uint32 BindingIndex;
-			if (!DescriptorInfo.GetDescriptorSetAndBindingIndex(FVulkanShaderHeader::UniformBuffer, Stage, BufferIndex, DescriptorSet, BindingIndex))
-			{
-				return;
-			}
+			return;
+		}
 
-			const FVulkanRealUniformBuffer* RealUniformBuffer = static_cast<const FVulkanRealUniformBuffer*>(UniformBuffer);
-			if (GDynamicGlobalUBs.GetValueOnAnyThread() > 1)
-			{
-				PendingGfxState->SetUniformBuffer<true>(DescriptorSet, BindingIndex, RealUniformBuffer);
-			}
-			else
-			{
-				PendingGfxState->SetUniformBuffer<false>(DescriptorSet, BindingIndex, RealUniformBuffer);
-			}
+		if (GDynamicGlobalUBs.GetValueOnAnyThread() > 1)
+		{
+			PendingGfxState->SetUniformBuffer<true>(DescriptorSet, BindingIndex, UniformBuffer);
 		}
 		else
 		{
-			const FVulkanEmulatedUniformBuffer* EmulatedUniformBuffer = static_cast<const FVulkanEmulatedUniformBuffer*>(UniformBuffer);
-			PendingGfxState->SetUniformBufferConstantData(Stage, BufferIndex, EmulatedUniformBuffer->ConstantData);
+			PendingGfxState->SetUniformBuffer<false>(DescriptorSet, BindingIndex, UniformBuffer);
 		}
 	}
 
@@ -468,7 +459,7 @@ inline void FVulkanCommandListContext::SetShaderUniformBuffer(ShaderStage::EStag
 	else
 	{
 		// Internal error: Completely empty UB!
-		checkSlow(!CodeHeader.bHasRealUBs || !HeaderUBInfo.bOnlyHasResources);
+		checkSlow(!HeaderUBInfo.bOnlyHasResources);
 	}
 }
 
@@ -508,36 +499,27 @@ void FVulkanCommandListContext::RHISetShaderUniformBuffer(FRHIComputeShader* Com
 	const FVulkanShaderHeader& CodeHeader = Shader->GetCodeHeader();
 	const FVulkanShaderHeader::FUniformBufferInfo& HeaderUBInfo = CodeHeader.UniformBuffers[BufferIndex];
 	checkfSlow(!HeaderUBInfo.LayoutHash || HeaderUBInfo.LayoutHash == UniformBuffer->GetLayout().GetHash(), TEXT("Mismatched UB layout! Got hash 0x%x, expected 0x%x!"), UniformBuffer->GetLayout().GetHash(), HeaderUBInfo.LayoutHash);
-	const bool bUseRealUBs = FVulkanPlatform::UseRealUBsOptimization(CodeHeader.bHasRealUBs != 0);
 
 	// Uniform Buffers
-	if (!bUseRealUBs || !HeaderUBInfo.bOnlyHasResources)
+	if (!HeaderUBInfo.bOnlyHasResources)
 	{
-		checkSlow(!bUseRealUBs || UniformBuffer->GetLayout().ConstantBufferSize > 0);
+		checkSlow(UniformBuffer->GetLayout().ConstantBufferSize > 0);
 		extern TAutoConsoleVariable<int32> GDynamicGlobalUBs;
-		if (bUseRealUBs)
+
+		uint8 DescriptorSet;
+		uint32 BindingIndex;
+		if (!DescriptorInfo.GetDescriptorSetAndBindingIndex(FVulkanShaderHeader::UniformBuffer, BufferIndex, DescriptorSet, BindingIndex))
 		{
-			uint8 DescriptorSet;
-			uint32 BindingIndex;
-			if (!DescriptorInfo.GetDescriptorSetAndBindingIndex(FVulkanShaderHeader::UniformBuffer, BufferIndex, DescriptorSet, BindingIndex))
-			{
-				return;
-			}
+			return;
+		}
 
-			const FVulkanRealUniformBuffer* RealUniformBuffer = static_cast<const FVulkanRealUniformBuffer*>(UniformBuffer);
-			if (GDynamicGlobalUBs.GetValueOnAnyThread() > 1)
-			{
-				State.SetUniformBuffer<true>(DescriptorSet, BindingIndex, RealUniformBuffer);
-			}
-			else
-			{
-				State.SetUniformBuffer<false>(DescriptorSet, BindingIndex, RealUniformBuffer);
-			}
+		if (GDynamicGlobalUBs.GetValueOnAnyThread() > 1)
+		{
+			State.SetUniformBuffer<true>(DescriptorSet, BindingIndex, UniformBuffer);
 		}
 		else
 		{
-			const FVulkanEmulatedUniformBuffer* EmulatedUniformBuffer = static_cast<const FVulkanEmulatedUniformBuffer*>(UniformBuffer);
-			State.SetUniformBufferConstantData(BufferIndex, EmulatedUniformBuffer->ConstantData);
+			State.SetUniformBuffer<false>(DescriptorSet, BindingIndex, UniformBuffer);
 		}
 	}
 
@@ -548,7 +530,7 @@ void FVulkanCommandListContext::RHISetShaderUniformBuffer(FRHIComputeShader* Com
 	else
 	{
 		// Internal error: Completely empty UB!
-		checkSlow(!CodeHeader.bHasRealUBs || !HeaderUBInfo.bOnlyHasResources);
+		checkSlow(!HeaderUBInfo.bOnlyHasResources);
 	}
 }
 
diff --git a/Engine/Source/Runtime/VulkanRHI/Private/VulkanGenericPlatform.h b/Engine/Source/Runtime/VulkanRHI/Private/VulkanGenericPlatform.h
--- a/Engine/Source/Runtime/VulkanRHI/Private/VulkanGenericPlatform.h
+++ b/Engine/Source/Runtime/VulkanRHI/Private/VulkanGenericPlatform.h
@@ -82,9 +82,6 @@ public:
 		InOutFeaturesToEnable.sparseResidencyAliased	= VK_FALSE;
 	}
 
-	// Some platforms only support real or non-real UBs, so this function can optimize it out
-	static bool UseRealUBsOptimization(bool bCodeHeaderUseRealUBs) { return bCodeHeaderUseRealUBs; }
-
 	static bool SupportParallelRenderingTasks() { return true; }
 
 	/** The status quo is false, so the default is chosen to not change it. As platforms opt in it may be better to flip the default. */
diff --git a/Engine/Source/Runtime/VulkanRHI/Private/VulkanPendingState.h b/Engine/Source/Runtime/VulkanRHI/Private/VulkanPendingState.h
--- a/Engine/Source/Runtime/VulkanRHI/Private/VulkanPendingState.h
+++ b/Engine/Source/Runtime/VulkanRHI/Private/VulkanPendingState.h
@@ -277,7 +277,7 @@ public:
 	}
 
 	template<bool bDynamic>
-	inline void SetUniformBuffer(uint8 DescriptorSet, uint32 BindingIndex, const FVulkanRealUniformBuffer* UniformBuffer)
+	inline void SetUniformBuffer(uint8 DescriptorSet, uint32 BindingIndex, const FVulkanUniformBuffer* UniformBuffer)
 	{
 		CurrentState->SetUniformBuffer<bDynamic>(DescriptorSet, BindingIndex, UniformBuffer);
 	}
diff --git a/Engine/Source/Runtime/VulkanRHI/Private/VulkanPipelineState.h b/Engine/Source/Runtime/VulkanRHI/Private/VulkanPipelineState.h
--- a/Engine/Source/Runtime/VulkanRHI/Private/VulkanPipelineState.h
+++ b/Engine/Source/Runtime/VulkanRHI/Private/VulkanPipelineState.h
@@ -93,7 +93,7 @@ public:
 	}
 
 	template<bool bDynamic>
-	inline void SetUniformBuffer(uint8 DescriptorSet, uint32 BindingIndex, const FVulkanRealUniformBuffer* UniformBuffer)
+	inline void SetUniformBuffer(uint8 DescriptorSet, uint32 BindingIndex, const FVulkanUniformBuffer* UniformBuffer)
 	{
 /*
 		if ((UniformBuffersWithDataMask[DescriptorSet] & (1ULL << (uint64)BindingIndex)) != 0)
@@ -237,7 +237,7 @@ public:
 		PackedUniformBuffers[Stage].SetEmulatedUniformBufferIntoPacked(BindingIndex, ConstantData, PackedUniformBuffersDirty[Stage]);
 	}
 
-	inline void SetDynamicUniformBuffer(uint8 DescriptorSet, uint32 BindingIndex, const FVulkanRealUniformBuffer* UniformBuffer)
+	inline void SetDynamicUniformBuffer(uint8 DescriptorSet, uint32 BindingIndex, const FVulkanUniformBuffer* UniformBuffer)
 	{
 		ensure(0);
 /*
diff --git a/Engine/Source/Runtime/VulkanRHI/Private/VulkanResources.h b/Engine/Source/Runtime/VulkanRHI/Private/VulkanResources.h
--- a/Engine/Source/Runtime/VulkanRHI/Private/VulkanResources.h
+++ b/Engine/Source/Runtime/VulkanRHI/Private/VulkanResources.h
@@ -1157,35 +1157,14 @@ protected:
 class FVulkanUniformBuffer : public FRHIUniformBuffer
 {
 public:
-	FVulkanUniformBuffer(const FRHIUniformBufferLayout& InLayout, const void* Contents, EUniformBufferUsage InUsage, EUniformBufferValidation Validation);
+	FVulkanUniformBuffer(FVulkanDevice& Device, const FRHIUniformBufferLayout& InLayout, const void* Contents, EUniformBufferUsage InUsage, EUniformBufferValidation Validation);
+	virtual ~FVulkanUniformBuffer();
 
 	const TArray<TRefCountPtr<FRHIResource>>& GetResourceTable() const { return ResourceTable; }
 
 	void UpdateResourceTable(const FRHIUniformBufferLayout& InLayout, const void* Contents, int32 ResourceNum);
 	void UpdateResourceTable(FRHIResource** Resources, int32 ResourceNum);
 
-protected:
-	TArray<TRefCountPtr<FRHIResource>> ResourceTable;
-};
-
-class FVulkanEmulatedUniformBuffer : public FVulkanUniformBuffer
-{
-public:
-	FVulkanEmulatedUniformBuffer(const FRHIUniformBufferLayout& InLayout, const void* Contents, EUniformBufferUsage InUsage, EUniformBufferValidation Validation);
-
-	TArray<uint8> ConstantData;
-
-	void UpdateConstantData(const void* Contents, int32 ContentsSize);
-};
-
-class FVulkanRealUniformBuffer : public FVulkanUniformBuffer
-{
-public:
-	FVulkanDevice* Device;
-	FVulkanRealUniformBuffer(FVulkanDevice& Device, const FRHIUniformBufferLayout& InLayout, const void* Contents, EUniformBufferUsage InUsage, EUniformBufferValidation Validation);
-	virtual ~FVulkanRealUniformBuffer();
-
-
 	inline uint32 GetOffset() const
 	{
 		return Allocation.Offset;
@@ -1195,7 +1174,14 @@ public:
 	{
 		NewAlloc.Swap(Allocation);
 	}
+	
+public:
+	FVulkanDevice* Device;
 	VulkanRHI::FVulkanAllocation Allocation;
+
+protected:
+	TArray<TRefCountPtr<FRHIResource>> ResourceTable;
+	
 };
 
 class FVulkanUnorderedAccessView : public FRHIUnorderedAccessView, public VulkanRHI::FVulkanViewBase
diff --git a/Engine/Source/Runtime/VulkanRHI/Private/VulkanShaders.cpp b/Engine/Source/Runtime/VulkanRHI/Private/VulkanShaders.cpp
--- a/Engine/Source/Runtime/VulkanRHI/Private/VulkanShaders.cpp
+++ b/Engine/Source/Runtime/VulkanRHI/Private/VulkanShaders.cpp
@@ -106,14 +106,8 @@ void FVulkanShader::Setup(TArrayView<const uint8> InShaderHeaderAndCode, uint64
 
 	SpirvSize = Spirv.Num() * sizeof(uint32);
 
-	if (CodeHeader.bHasRealUBs)
-	{
-		check(CodeHeader.UniformBufferSpirvInfos.Num() == CodeHeader.UniformBuffers.Num());
-	}
-	else
-	{
-		checkSlow(CodeHeader.UniformBufferSpirvInfos.Num() == 0);
-	}
+	check(CodeHeader.UniformBufferSpirvInfos.Num() == CodeHeader.UniformBuffers.Num());
+
 	check(CodeHeader.GlobalSpirvInfos.Num() == CodeHeader.Globals.Num());
 
 	StaticSlots.Reserve(CodeHeader.UniformBuffers.Num());
@@ -288,21 +282,19 @@ void FVulkanLayout::PatchSpirvBindings(TArray<uint32>& Spirv, EShaderFrequency F
 	//#todo-rco: Do we need an actual copy of the SPIR-V?
 	ShaderStage::EStage Stage = ShaderStage::GetStageForFrequency(Frequency);
 	const FDescriptorSetRemappingInfo::FStageInfo& StageInfo = DescriptorSetLayout.RemappingInfo.StageInfos[Stage];
-	if (CodeHeader.bHasRealUBs)
+	
+	checkSlow(StageInfo.UniformBuffers.Num() == CodeHeader.UniformBufferSpirvInfos.Num());
+	for (int32 Index = 0; Index < CodeHeader.UniformBufferSpirvInfos.Num(); ++Index)
 	{
-		checkSlow(StageInfo.UniformBuffers.Num() == CodeHeader.UniformBufferSpirvInfos.Num());
-		for (int32 Index = 0; Index < CodeHeader.UniformBufferSpirvInfos.Num(); ++Index)
+		if (StageInfo.UniformBuffers[Index].bHasConstantData)
 		{
-			if (StageInfo.UniformBuffers[Index].bHasConstantData)
-			{
-				const uint32 OffsetDescriptorSet = CodeHeader.UniformBufferSpirvInfos[Index].DescriptorSetOffset;
-				const uint32 OffsetBindingIndex = CodeHeader.UniformBufferSpirvInfos[Index].BindingIndexOffset;
-				check(OffsetDescriptorSet != UINT32_MAX && OffsetBindingIndex != UINT32_MAX);
-				uint16 NewDescriptorSet = StageInfo.UniformBuffers[Index].Remapping.NewDescriptorSet;
-				Spirv[OffsetDescriptorSet] = NewDescriptorSet;
-				uint16 NewBindingIndex = StageInfo.UniformBuffers[Index].Remapping.NewBindingIndex;
-				Spirv[OffsetBindingIndex] = NewBindingIndex;
-			}
+			const uint32 OffsetDescriptorSet = CodeHeader.UniformBufferSpirvInfos[Index].DescriptorSetOffset;
+			const uint32 OffsetBindingIndex = CodeHeader.UniformBufferSpirvInfos[Index].BindingIndexOffset;
+			check(OffsetDescriptorSet != UINT32_MAX && OffsetBindingIndex != UINT32_MAX);
+			uint16 NewDescriptorSet = StageInfo.UniformBuffers[Index].Remapping.NewDescriptorSet;
+			Spirv[OffsetDescriptorSet] = NewDescriptorSet;
+			uint16 NewBindingIndex = StageInfo.UniformBuffers[Index].Remapping.NewBindingIndex;
+			Spirv[OffsetBindingIndex] = NewBindingIndex;
 		}
 	}
 
@@ -564,58 +556,55 @@ void FVulkanDescriptorSetsLayoutInfo::FinalizeBindings(const FUniformBufferGathe
 				AddDescriptor(DescriptorSet, Binding);
 			}
 
-			if (ShaderHeader->bHasRealUBs)
+			RemappingInfo.StageInfos[Stage].UniformBuffers.Reserve(ShaderHeader->UniformBuffers.Num());
+			for (int32 Index = 0; Index < ShaderHeader->UniformBuffers.Num(); ++Index)
 			{
-				RemappingInfo.StageInfos[Stage].UniformBuffers.Reserve(ShaderHeader->UniformBuffers.Num());
-				for (int32 Index = 0; Index < ShaderHeader->UniformBuffers.Num(); ++Index)
+				VkDescriptorType Type = bConvertAllUBsToDynamic ? VK_DESCRIPTOR_TYPE_UNIFORM_BUFFER_DYNAMIC : VK_DESCRIPTOR_TYPE_UNIFORM_BUFFER;
+				// Here we might mess up with the stageFlags, so reset them every loop
+				Binding.stageFlags = StageFlags;
+				Binding.descriptorType = Type;
+				const FVulkanShaderHeader::FUniformBufferInfo& UBInfo = ShaderHeader->UniformBuffers[Index];
+				const uint32 LayoutHash = UBInfo.LayoutHash;
+				const bool bUBHasConstantData = UBInfo.ConstantDataOriginalBindingIndex != UINT16_MAX;
+				if (bUBHasConstantData)
 				{
-					VkDescriptorType Type = bConvertAllUBsToDynamic ? VK_DESCRIPTOR_TYPE_UNIFORM_BUFFER_DYNAMIC : VK_DESCRIPTOR_TYPE_UNIFORM_BUFFER;
-					// Here we might mess up with the stageFlags, so reset them every loop
-					Binding.stageFlags = StageFlags;
-					Binding.descriptorType = Type;
-					const FVulkanShaderHeader::FUniformBufferInfo& UBInfo = ShaderHeader->UniformBuffers[Index];
-					const uint32 LayoutHash = UBInfo.LayoutHash;
-					const bool bUBHasConstantData = UBInfo.ConstantDataOriginalBindingIndex != UINT16_MAX;
-					if (bUBHasConstantData)
+					bool bProcessRegularUB = true;
+					const VkShaderStageFlags* FoundFlags = bMoveCommonUBsToExtraSet ? UBGatherInfo.CommonUBLayoutsToStageMap.Find(LayoutHash) : nullptr;
+					if (FoundFlags)
 					{
-						bool bProcessRegularUB = true;
-						const VkShaderStageFlags* FoundFlags = bMoveCommonUBsToExtraSet ? UBGatherInfo.CommonUBLayoutsToStageMap.Find(LayoutHash) : nullptr;
-						if (FoundFlags)
+						if (const FDescriptorSetRemappingInfo::FUBRemappingInfo* UBRemapInfo = AlreadyProcessedUBs.Find(LayoutHash))
 						{
-							if (const FDescriptorSetRemappingInfo::FUBRemappingInfo* UBRemapInfo = AlreadyProcessedUBs.Find(LayoutHash))
-							{
-								RemappingInfo.AddRedundantUB(Stage, Index, UBRemapInfo);
-							}
-							else
-							{
-								//#todo-rco: Only process constant data part of the UB
-								check(bUBHasConstantData);
-
-								Binding.stageFlags = *FoundFlags;
-								uint32 NewBindingIndex;
-								AlreadyProcessedUBs.Add(LayoutHash, RemappingInfo.AddUBWithData(Stage, Index, CommonUBDescriptorSet, Type, NewBindingIndex));
-								Binding.binding = NewBindingIndex;
-
-								AddDescriptor(CommonUBDescriptorSet, Binding);
-							}
-							bProcessRegularUB = false;
+							RemappingInfo.AddRedundantUB(Stage, Index, UBRemapInfo);
 						}
-
-						if (bProcessRegularUB)
+						else
 						{
-							int32 DescriptorSet = FindOrAddDescriptorSet(Stage);
+							//#todo-rco: Only process constant data part of the UB
+							check(bUBHasConstantData);
+
+							Binding.stageFlags = *FoundFlags;
 							uint32 NewBindingIndex;
-							RemappingInfo.AddUBWithData(Stage, Index, DescriptorSet, Type, NewBindingIndex);
+							AlreadyProcessedUBs.Add(LayoutHash, RemappingInfo.AddUBWithData(Stage, Index, CommonUBDescriptorSet, Type, NewBindingIndex));
 							Binding.binding = NewBindingIndex;
 
-							AddDescriptor(FindOrAddDescriptorSet(Stage), Binding);
+							AddDescriptor(CommonUBDescriptorSet, Binding);
 						}
+						bProcessRegularUB = false;
 					}
-					else
+
+					if (bProcessRegularUB)
 					{
-						RemappingInfo.AddUBResourceOnly(Stage, Index);
+						int32 DescriptorSet = FindOrAddDescriptorSet(Stage);
+						uint32 NewBindingIndex;
+						RemappingInfo.AddUBWithData(Stage, Index, DescriptorSet, Type, NewBindingIndex);
+						Binding.binding = NewBindingIndex;
+
+						AddDescriptor(FindOrAddDescriptorSet(Stage), Binding);
 					}
 				}
+				else
+				{
+					RemappingInfo.AddUBResourceOnly(Stage, Index);
+				}
 			}
 
 			RemappingInfo.StageInfos[Stage].Globals.Reserve(ShaderHeader->Globals.Num());
diff --git a/Engine/Source/Runtime/VulkanRHI/Private/VulkanUniformBuffer.cpp b/Engine/Source/Runtime/VulkanRHI/Private/VulkanUniformBuffer.cpp
--- a/Engine/Source/Runtime/VulkanRHI/Private/VulkanUniformBuffer.cpp
+++ b/Engine/Source/Runtime/VulkanRHI/Private/VulkanUniformBuffer.cpp
@@ -44,8 +44,8 @@ static inline EBufferUsageFlags UniformBufferToBufferUsage(EUniformBufferUsage U
 	}
 }
 
-FVulkanUniformBuffer::FVulkanUniformBuffer(const FRHIUniformBufferLayout& InLayout, const void* Contents, EUniformBufferUsage InUsage, EUniformBufferValidation Validation)
-	: FRHIUniformBuffer(InLayout)
+FVulkanUniformBuffer::FVulkanUniformBuffer(FVulkanDevice& Device, const FRHIUniformBufferLayout& InLayout, const void* Contents, EUniformBufferUsage InUsage, EUniformBufferValidation Validation)
+	: FRHIUniformBuffer(InLayout), Device(&Device)
 {
 #if VULKAN_ENABLE_AGGRESSIVE_STATS
 	SCOPE_CYCLE_COUNTER(STAT_VulkanUniformBufferCreateTime);
@@ -68,6 +68,19 @@ FVulkanUniformBuffer::FVulkanUniformBuffer(const FRHIUniformBufferLayout& InLayo
 			ResourceTable[Index] = GetShaderParameterResourceRHI(Contents, InLayout.Resources[Index].MemberOffset, InLayout.Resources[Index].MemberType);
 		}
 	}
+
+	if (InLayout.ConstantBufferSize > 0)
+	{
+		VulkanRHI::FMemoryManager& ResourceMgr = Device.GetMemoryManager();
+
+		// Set it directly as there is no previous one
+		ResourceMgr.AllocUniformBuffer(Allocation, InLayout.ConstantBufferSize, Contents);
+	}
+}
+
+FVulkanUniformBuffer::~FVulkanUniformBuffer()
+{
+	Device->GetMemoryManager().FreeUniformBuffer(Allocation);
 }
 
 void FVulkanUniformBuffer::UpdateResourceTable(const FRHIUniformBufferLayout& InLayout, const void* Contents, int32 NumResources)
@@ -92,79 +105,13 @@ void FVulkanUniformBuffer::UpdateResourceTable(FRHIResource** Resources, int32 R
 }
 
 
-FVulkanEmulatedUniformBuffer::FVulkanEmulatedUniformBuffer(const FRHIUniformBufferLayout& InLayout, const void* Contents, EUniformBufferUsage InUsage, EUniformBufferValidation Validation)
-	: FVulkanUniformBuffer(InLayout, Contents, InUsage, Validation)
-{
-#if VULKAN_ENABLE_AGGRESSIVE_STATS
-	SCOPE_CYCLE_COUNTER(STAT_VulkanUniformBufferCreateTime);
-#endif
-
-	// Contents might be null but size > 0 as the data doesn't need a CPU copy
-	if (InLayout.ConstantBufferSize)
-	{
-		// Create uniform buffer, which is stored on the CPU, the buffer is uploaded to a correct GPU buffer in UpdateDescriptorSets()
-		ConstantData.AddUninitialized(InLayout.ConstantBufferSize);
-		if (Contents)
-		{
-			FMemory::Memcpy(ConstantData.GetData(), Contents, InLayout.ConstantBufferSize);
-		}
-	}
-
-	// Ancestor's constructor will set up the Resource table, so nothing else to do here
-}
-
-void FVulkanEmulatedUniformBuffer::UpdateConstantData(const void* Contents, int32 ContentsSize)
-{
-	checkSlow(ConstantData.Num() * sizeof(ConstantData[0]) == ContentsSize);
-	if (ContentsSize > 0)
-	{
-		FMemory::Memcpy(ConstantData.GetData(), Contents, ContentsSize);
-	}
-}
-
-
-FVulkanRealUniformBuffer::FVulkanRealUniformBuffer(FVulkanDevice& Device, const FRHIUniformBufferLayout& InLayout, const void* Contents, EUniformBufferUsage InUsage, EUniformBufferValidation Validation)
-	: FVulkanUniformBuffer(InLayout, Contents, InUsage, Validation)
-	, Device(&Device)
-{
-#if VULKAN_ENABLE_AGGRESSIVE_STATS
-	SCOPE_CYCLE_COUNTER(STAT_VulkanUniformBufferCreateTime);
-#endif
-
-	if (InLayout.ConstantBufferSize > 0)
-	{
-		VulkanRHI::FMemoryManager& ResourceMgr = Device.GetMemoryManager();
-
-		// Set it directly as there is no previous one
-		ResourceMgr.AllocUniformBuffer(Allocation, InLayout.ConstantBufferSize, Contents);
-	}
-
-	// Ancestor's constructor will set up the Resource table, so nothing else to do here
-}
-
-FVulkanRealUniformBuffer::~FVulkanRealUniformBuffer()
-{
-	Device->GetMemoryManager().FreeUniformBuffer(Allocation);
-}
-
 FUniformBufferRHIRef FVulkanDynamicRHI::RHICreateUniformBuffer(const void* Contents, const FRHIUniformBufferLayout& Layout, EUniformBufferUsage Usage, EUniformBufferValidation Validation)
 {
 	LLM_SCOPE_VULKAN(ELLMTagVulkan::VulkanUniformBuffers);
 
-	static TConsoleVariableData<int32>* CVar = IConsoleManager::Get().FindTConsoleVariableDataInt(TEXT("r.Vulkan.UseRealUBs"));
-	const bool bHasRealUBs = FVulkanPlatform::UseRealUBsOptimization(CVar && CVar->GetValueOnAnyThread() > 0);
-	if (bHasRealUBs)
-	{
-		return new FVulkanRealUniformBuffer(*Device, Layout, Contents, Usage, Validation);
-	}
-	else
-	{
-		// Parts of the buffer are later on copied for each shader stage into the packed uniform buffer
-		return new FVulkanEmulatedUniformBuffer(Layout, Contents, Usage, Validation);
-	}
+	return new FVulkanUniformBuffer(*Device, Layout, Contents, Usage, Validation);
 }
 
-template <bool bRealUBs>
 inline void FVulkanDynamicRHI::UpdateUniformBuffer(FVulkanUniformBuffer* UniformBuffer, const void* Contents)
 {
 	SCOPE_CYCLE_COUNTER(STAT_VulkanUpdateUniformBuffers);
@@ -175,13 +122,10 @@ inline void FVulkanDynamicRHI::UpdateUniformBuffer(FVulkanUniformBuffer* Uniform
 
 	FRHICommandListImmediate& RHICmdList = FRHICommandListExecutor::GetImmediateCommandList();
 
-	FVulkanRealUniformBuffer* RealUniformBuffer = bRealUBs ? (FVulkanRealUniformBuffer*)UniformBuffer : nullptr;
-	FVulkanEmulatedUniformBuffer* EmulatedUniformBuffer = bRealUBs ? nullptr : (FVulkanEmulatedUniformBuffer*)UniformBuffer;
-
 	FVulkanAllocation NewUBAlloc;
 	bool bUseUpload = GVulkanAllowUniformUpload && !RHICmdList.IsInsideRenderPass(); //inside renderpasses, a rename is enforced.
 
-	if (bRealUBs && !bUseUpload)
+	if (!bUseUpload)
 	{
 		if (ConstantBufferSize > 0)
 		{
@@ -190,7 +134,7 @@ inline void FVulkanDynamicRHI::UpdateUniformBuffer(FVulkanUniformBuffer* Uniform
 		}
 	}
 
-	auto UpdateUniformBufferHelper = [](FVulkanCommandListContext& Context, FVulkanRealUniformBuffer* VulkanUniformBuffer, int32 DataSize, const void* Data)
+	auto UpdateUniformBufferHelper = [](FVulkanCommandListContext& Context, FVulkanUniformBuffer* VulkanUniformBuffer, int32 DataSize, const void* Data)
 	{
 		FVulkanCmdBuffer* CmdBuffer = Context.GetCommandBufferManager()->GetActiveCmdBufferDirect();
 		ensure(CmdBuffer->IsOutsideRenderPass());
@@ -222,26 +166,20 @@ inline void FVulkanDynamicRHI::UpdateUniformBuffer(FVulkanUniformBuffer* Uniform
 	bool bRHIBypass = RHICmdList.Bypass();
 	if (bRHIBypass)
 	{
-		if (bRealUBs)
+		if (ConstantBufferSize > 0)
 		{
-			if (ConstantBufferSize > 0)
+			if(bUseUpload)
+			{			
+				FVulkanCommandListContext& Context = Device->GetImmediateContext();
+				UpdateUniformBufferHelper(Context, UniformBuffer, ConstantBufferSize, Contents);
+			}
+			else
 			{
-				if(bUseUpload)
-				{			
-					FVulkanCommandListContext& Context = Device->GetImmediateContext();
-					UpdateUniformBufferHelper(Context, RealUniformBuffer, ConstantBufferSize, Contents);
-				}
-				else
-				{
-					RealUniformBuffer->UpdateAllocation(NewUBAlloc);
-					Device->GetMemoryManager().FreeUniformBuffer(NewUBAlloc);
-				}
+				UniformBuffer->UpdateAllocation(NewUBAlloc);
+				Device->GetMemoryManager().FreeUniformBuffer(NewUBAlloc);
 			}
 		}
-		else
-		{
-			EmulatedUniformBuffer->UpdateConstantData(Contents, ConstantBufferSize);
-		}
+
 		UniformBuffer->UpdateResourceTable(Layout, Contents, NumResources);
 	}
 	else
@@ -257,44 +195,32 @@ inline void FVulkanDynamicRHI::UpdateUniformBuffer(FVulkanUniformBuffer* Uniform
 			}
 		}
 
-		if (bRealUBs)
+		if(bUseUpload)
 		{
-			if(bUseUpload)
-			{
-				void* CmdListConstantBufferData = RHICmdList.Alloc(ConstantBufferSize, 16);
-				FMemory::Memcpy(CmdListConstantBufferData, Contents, ConstantBufferSize);
-
-				RHICmdList.EnqueueLambda([UpdateUniformBufferHelper, RealUniformBuffer, CmdListResources, NumResources, ConstantBufferSize, CmdListConstantBufferData](FRHICommandList& CmdList)
-				{
-					FVulkanCommandListContext& Context = (FVulkanCommandListContext&)CmdList.GetContext().GetLowestLevelContext();
-					UpdateUniformBufferHelper(Context, RealUniformBuffer, ConstantBufferSize, CmdListConstantBufferData);
-					RealUniformBuffer->UpdateResourceTable(CmdListResources, NumResources);
-				});
-			}
-			else
+			void* CmdListConstantBufferData = RHICmdList.Alloc(ConstantBufferSize, 16);
+			FMemory::Memcpy(CmdListConstantBufferData, Contents, ConstantBufferSize);
+
+			RHICmdList.EnqueueLambda([UpdateUniformBufferHelper, UniformBuffer, CmdListResources, NumResources, ConstantBufferSize, CmdListConstantBufferData](FRHICommandList& CmdList)
 			{
-				NewUBAlloc.Disown(); //this releases ownership while its put into the lambda
-				RHICmdList.EnqueueLambda([RealUniformBuffer, NewUBAlloc, CmdListResources, NumResources](FRHICommandList& CmdList)
-				{
-					FVulkanAllocation Alloc;
-					Alloc.Reference(NewUBAlloc);
-					Alloc.Own(); //this takes ownership of the allocation
-					RealUniformBuffer->UpdateAllocation(Alloc);
-					RealUniformBuffer->Device->GetMemoryManager().FreeUniformBuffer(Alloc);
-					RealUniformBuffer->UpdateResourceTable(CmdListResources, NumResources);
-				});
-			}
+				FVulkanCommandListContext& Context = (FVulkanCommandListContext&)CmdList.GetContext().GetLowestLevelContext();
+				UpdateUniformBufferHelper(Context, UniformBuffer, ConstantBufferSize, CmdListConstantBufferData);
+				UniformBuffer->UpdateResourceTable(CmdListResources, NumResources);
+			});
 		}
 		else
 		{
-			void* CmdListConstantBufferData = RHICmdList.Alloc(ConstantBufferSize, 16);
-			FMemory::Memcpy(CmdListConstantBufferData, Contents, ConstantBufferSize);
-			RHICmdList.EnqueueLambda([EmulatedUniformBuffer, CmdListResources, NumResources, CmdListConstantBufferData, ConstantBufferSize](FRHICommandList&)
+			NewUBAlloc.Disown(); //this releases ownership while its put into the lambda
+			RHICmdList.EnqueueLambda([UniformBuffer, NewUBAlloc, CmdListResources, NumResources](FRHICommandList& CmdList)
 			{
-				EmulatedUniformBuffer->UpdateConstantData(CmdListConstantBufferData, ConstantBufferSize);
-				EmulatedUniformBuffer->UpdateResourceTable(CmdListResources, NumResources);
+				FVulkanAllocation Alloc;
+				Alloc.Reference(NewUBAlloc);
+				Alloc.Own(); //this takes ownership of the allocation
+				UniformBuffer->UpdateAllocation(Alloc);
+				UniformBuffer->Device->GetMemoryManager().FreeUniformBuffer(Alloc);
+				UniformBuffer->UpdateResourceTable(CmdListResources, NumResources);
 			});
 		}
+		
 		RHICmdList.RHIThreadFence(true);
 	}
 }
@@ -302,17 +228,8 @@ inline void FVulkanDynamicRHI::UpdateUniformBuffer(FVulkanUniformBuffer* Uniform
 
 void FVulkanDynamicRHI::RHIUpdateUniformBuffer(FRHIUniformBuffer* UniformBufferRHI, const void* Contents)
 {
-	static TConsoleVariableData<int32>* CVar = IConsoleManager::Get().FindTConsoleVariableDataInt(TEXT("r.Vulkan.UseRealUBs"));
-	const bool bHasRealUBs = FVulkanPlatform::UseRealUBsOptimization(CVar && CVar->GetValueOnAnyThread() > 0);
 	FVulkanUniformBuffer* UniformBuffer = ResourceCast(UniformBufferRHI);
-	if (bHasRealUBs)
-	{
-		UpdateUniformBuffer<true>(UniformBuffer, Contents);
-	}
-	else
-	{
-		UpdateUniformBuffer<false>(UniformBuffer, Contents);
-	}
+	UpdateUniformBuffer(UniformBuffer, Contents);
 }
 
 FVulkanUniformBufferUploader::FVulkanUniformBufferUploader(FVulkanDevice* InDevice)
diff --git a/Engine/Source/Runtime/VulkanRHI/Private/Windows/VulkanWindowsPlatform.h b/Engine/Source/Runtime/VulkanRHI/Private/Windows/VulkanWindowsPlatform.h
--- a/Engine/Source/Runtime/VulkanRHI/Private/Windows/VulkanWindowsPlatform.h
+++ b/Engine/Source/Runtime/VulkanRHI/Private/Windows/VulkanWindowsPlatform.h
@@ -66,13 +66,6 @@ public:
 
 	static void WriteCrashMarker(const FOptionalVulkanDeviceExtensions& OptionalExtensions, VkCommandBuffer CmdBuffer, VkBuffer DestBuffer, const TArrayView<uint32>& Entries, bool bAdding);
 
-	// Some platforms only support real or non-real UBs, so this function can optimize it out
-	static bool UseRealUBsOptimization(bool bCodeHeaderUseRealUBs)
-	{
-		static auto* CVar = IConsoleManager::Get().FindTConsoleVariableDataInt(TEXT("r.Vulkan.UseRealUBs"));
-		return (CVar && CVar->GetValueOnAnyThread() == 0) ? false : bCodeHeaderUseRealUBs;
-	}
-
 	static void EnablePhysicalDeviceFeatureExtensions(VkDeviceCreateInfo& DeviceInfo, FVulkanDevice& Device);
 };
 
diff --git a/Engine/Source/Runtime/VulkanRHI/Public/VulkanDynamicRHI.h b/Engine/Source/Runtime/VulkanRHI/Public/VulkanDynamicRHI.h
--- a/Engine/Source/Runtime/VulkanRHI/Public/VulkanDynamicRHI.h
+++ b/Engine/Source/Runtime/VulkanRHI/Public/VulkanDynamicRHI.h
@@ -459,7 +459,6 @@ protected:
 	void InternalUpdateTexture2D(bool bFromRenderingThread, FRHITexture2D* TextureRHI, uint32 MipIndex, const struct FUpdateTextureRegion2D& UpdateRegion, uint32 SourcePitch, const uint8* SourceData);
 	void InternalUpdateTexture3D(bool bFromRenderingThread, FRHITexture3D* TextureRHI, uint32 MipIndex, const struct FUpdateTextureRegion3D& UpdateRegion, uint32 SourceRowPitch, uint32 SourceDepthPitch, const uint8* SourceData);
 
-	template <bool bRealUBs>
 	void UpdateUniformBuffer(FVulkanUniformBuffer* UniformBuffer, const void* Contents);
 
 	static void SetupValidationRequests();
diff --git a/Engine/Source/Runtime/VulkanRHI/Public/VulkanShaderResources.h b/Engine/Source/Runtime/VulkanRHI/Public/VulkanShaderResources.h
--- a/Engine/Source/Runtime/VulkanRHI/Public/VulkanShaderResources.h
+++ b/Engine/Source/Runtime/VulkanRHI/Public/VulkanShaderResources.h
@@ -183,10 +183,6 @@ struct FVulkanShaderHeader
 	// Mostly relevant for Vertex Shaders
 	uint32									InOutMask;
 
-	bool									bHasRealUBs;
-	uint8									Pad0 = 0;
-	uint16									Pad1 = 1;
-
 	FSHAHash								SourceHash;
 	uint32									SpirvCRC = 0;
 
@@ -202,7 +198,6 @@ struct FVulkanShaderHeader
 	};
 	FVulkanShaderHeader(EInit)
 		: InOutMask(0)
-		, bHasRealUBs(0)
 	{
 	}
 };
@@ -290,7 +285,6 @@ inline FArchive& operator<<(FArchive& Ar, FVulkanShaderHeader& Header)
 	Ar << Header.EmulatedUBCopyRanges;
 	Ar << Header.EmulatedUBsCopyInfo;
 	Ar << Header.InOutMask;
-	Ar << Header.bHasRealUBs;
 	Ar << Header.SourceHash;
 	Ar << Header.SpirvCRC;
 	Ar << Header.UniformBufferSpirvInfos;
