UE5 - Fix Vulkan missing dual blending annotation; move it to be platform agnostic
#rnx

diff --git a/Engine/Shaders/Public/Platform.ush b/Engine/Shaders/Public/Platform.ush
--- a/Engine/Shaders/Public/Platform.ush
+++ b/Engine/Shaders/Public/Platform.ush
@@ -77,6 +77,10 @@
 #define COMPILER_SUPPORTS_QUAD_PASS 0
 #endif
 
+#ifndef COMPILER_SUPPORTS_DUAL_SOURCE_BLENDING_SLOT_DECORATION
+#define COMPILER_SUPPORTS_DUAL_SOURCE_BLENDING_SLOT_DECORATION	0
+#endif
+
 #ifndef COMPILER_SUPPORTS_PRIMITIVE_SHADERS
 #define COMPILER_SUPPORTS_PRIMITIVE_SHADERS 0
 #endif
@@ -490,7 +494,7 @@ void ShaderYield()
 #define EARLYDEPTHSTENCIL [earlydepthstencil]
 #endif
 
-#if METAL_MRT_PROFILE || METAL_SM5_PROFILE
+#if COMPILER_SUPPORTS_DUAL_SOURCE_BLENDING_SLOT_DECORATION
 #define DUAL_SOURCE_BLENDING_SLOT(SLOT) [[vk::location(0), vk::index(SLOT)]]
 #endif
 
diff --git a/Engine/Source/Developer/Apple/MetalShaderFormat/Private/MetalShaderCompiler.cpp b/Engine/Source/Developer/Apple/MetalShaderFormat/Private/MetalShaderCompiler.cpp
--- a/Engine/Source/Developer/Apple/MetalShaderFormat/Private/MetalShaderCompiler.cpp
+++ b/Engine/Source/Developer/Apple/MetalShaderFormat/Private/MetalShaderCompiler.cpp
@@ -1118,6 +1118,8 @@ void CompileShader_Metal(const FShaderCompilerInput& _Input,FShaderCompilerOutpu
 		AdditionalDefines.SetDefine(TEXT("COMPILER_SUPPORTS_ATTRIBUTES"), (uint32)1);
 	}
 
+	AdditionalDefines.SetDefine(TEXT("COMPILER_SUPPORTS_DUAL_SOURCE_BLENDING_SLOT_DECORATION"), (uint32)1);
+
 	if (Input.bSkipPreprocessedCache)
 	{
 		if (!FFileHelper::LoadFileToString(PreprocessedShader, *Input.VirtualSourceFilePath))
diff --git a/Engine/Source/Developer/VulkanShaderFormat/Private/VulkanShaderCompiler.cpp b/Engine/Source/Developer/VulkanShaderFormat/Private/VulkanShaderCompiler.cpp
--- a/Engine/Source/Developer/VulkanShaderFormat/Private/VulkanShaderCompiler.cpp
+++ b/Engine/Source/Developer/VulkanShaderFormat/Private/VulkanShaderCompiler.cpp
@@ -2407,6 +2407,7 @@ void DoCompileVulkanShader(const FShaderCompilerInput& Input, FShaderCompilerOut
 	AdditionalDefines.SetDefine(TEXT("row_major"), TEXT(""));
 
 	AdditionalDefines.SetDefine(TEXT("COMPILER_SUPPORTS_ATTRIBUTES"), (uint32)1);
+	AdditionalDefines.SetDefine(TEXT("COMPILER_SUPPORTS_DUAL_SOURCE_BLENDING_SLOT_DECORATION"), (uint32)1);
 
 	const bool bUseFullPrecisionInPS = Input.Environment.CompilerFlags.Contains(CFLAG_UseFullPrecisionInPS);
 	if (bUseFullPrecisionInPS)
