Only emit precision qualifier for integer and floating-point types in ESSL output (Vulkan backend).

#rnx

diff --git a/Engine/Source/Developer/VulkanShaderFormat/Private/VulkanBackend.cpp b/Engine/Source/Developer/VulkanShaderFormat/Private/VulkanBackend.cpp
--- a/Engine/Source/Developer/VulkanShaderFormat/Private/VulkanBackend.cpp
+++ b/Engine/Source/Developer/VulkanShaderFormat/Private/VulkanBackend.cpp
@@ -223,6 +223,20 @@ static bool UsesUEIntrinsic(exec_list* Instructions, const char * UEIntrinsic)
 	return Visitor.bFound;
 }
 
+static bool GlslTypeSupportsPrecisionQualifier(glsl_base_type GlslBaseType)
+{
+	switch (GlslBaseType)
+	{
+	case GLSL_TYPE_UINT:
+	case GLSL_TYPE_INT:
+	case GLSL_TYPE_HALF:
+	case GLSL_TYPE_FLOAT:
+		return true;
+	default:
+		return false;
+	}
+}
+
 static inline void Scanf(char* Dest, const char* Format, float* OutValue)
 {
 #if PLATFORM_WINDOWS
@@ -2774,7 +2788,8 @@ class FGenerateVulkanVisitor : public ir_visitor
 			{
 				for (unsigned j = 0; j < s->length; j++)
 				{
-					ralloc_asprintf_append(buffer, "\t%s ", (state->language_version == 310 && bEmitPrecision) ? "highp" : "");
+					const bool bSupportsPrecisionQualifier = GlslTypeSupportsPrecisionQualifier(s->fields.structure[j].type->base_type);
+					ralloc_asprintf_append(buffer, "\t%s ", (state->language_version == 310 && bEmitPrecision && bSupportsPrecisionQualifier) ? "highp" : "");
 					print_type_pre(s->fields.structure[j].type);
 					ralloc_asprintf_append(buffer, " %s", s->fields.structure[j].name);
 					print_type_post(s->fields.structure[j].type);
@@ -2845,8 +2860,8 @@ class FGenerateVulkanVisitor : public ir_visitor
 						//EHart - name-mangle variables to prevent colliding names
 						//#todo-rco: Check if this is still is needed when creating PSOs
 						//ralloc_asprintf_append(buffer, "#define %s %s%s\n", var->name, var->name, block_name);
-						bool bIsBoolType = var->type->base_type == GLSL_TYPE_BOOL;
-						ralloc_asprintf_append(buffer, "\t%s", (state->language_version == 310 && bEmitPrecision && !bIsBoolType) ? "highp " : "");
+						const bool bSupportsPrecisionQualifier = GlslTypeSupportsPrecisionQualifier(var->type->base_type);
+						ralloc_asprintf_append(buffer, "\t%s", (state->language_version == 310 && bEmitPrecision && bSupportsPrecisionQualifier) ? "highp " : "");
 						print_type_pre(var->type);
 						ralloc_asprintf_append(buffer, " %s", var->name);
 						print_type_post(var->type);
