Fix precision qualifier condition for ESSL output. Fix for CL 16757096.

#rnx

diff --git a/Engine/Source/Developer/VulkanShaderFormat/Private/VulkanBackend.cpp b/Engine/Source/Developer/VulkanShaderFormat/Private/VulkanBackend.cpp
--- a/Engine/Source/Developer/VulkanShaderFormat/Private/VulkanBackend.cpp
+++ b/Engine/Source/Developer/VulkanShaderFormat/Private/VulkanBackend.cpp
@@ -223,15 +223,21 @@ static bool UsesUEIntrinsic(exec_list* Instructions, const char * UEIntrinsic)
 	return Visitor.bFound;
 }
 
-static bool GlslTypeSupportsPrecisionQualifier(glsl_base_type GlslBaseType)
+static bool GlslTypeSupportsPrecisionQualifier(const glsl_type& GlslType)
 {
-	switch (GlslBaseType)
+	switch (GlslType.base_type)
 	{
 	case GLSL_TYPE_UINT:
 	case GLSL_TYPE_INT:
 	case GLSL_TYPE_HALF:
 	case GLSL_TYPE_FLOAT:
 		return true;
+	case GLSL_TYPE_ARRAY:
+		if (const glsl_type* GlslArrayElementType = GlslType.fields.array)
+		{
+			return GlslTypeSupportsPrecisionQualifier(*GlslArrayElementType);
+		}
+		return false;
 	default:
 		return false;
 	}
@@ -2786,9 +2792,9 @@ class FGenerateVulkanVisitor : public ir_visitor
 			}
 			else
 			{
-				for (unsigned j = 0; j < s->length; j++)
+				for (unsigned j = 0; j < s->length; j++) 
 				{
-					const bool bSupportsPrecisionQualifier = GlslTypeSupportsPrecisionQualifier(s->fields.structure[j].type->base_type);
+					const bool bSupportsPrecisionQualifier = GlslTypeSupportsPrecisionQualifier(*(s->fields.structure[j].type));
 					ralloc_asprintf_append(buffer, "\t%s ", (state->language_version == 310 && bEmitPrecision && bSupportsPrecisionQualifier) ? "highp" : "");
 					print_type_pre(s->fields.structure[j].type);
 					ralloc_asprintf_append(buffer, " %s", s->fields.structure[j].name);
@@ -2860,7 +2866,7 @@ class FGenerateVulkanVisitor : public ir_visitor
 						//EHart - name-mangle variables to prevent colliding names
 						//#todo-rco: Check if this is still is needed when creating PSOs
 						//ralloc_asprintf_append(buffer, "#define %s %s%s\n", var->name, var->name, block_name);
-						const bool bSupportsPrecisionQualifier = GlslTypeSupportsPrecisionQualifier(var->type->base_type);
+						const bool bSupportsPrecisionQualifier = GlslTypeSupportsPrecisionQualifier(*(var->type));
 						ralloc_asprintf_append(buffer, "\t%s", (state->language_version == 310 && bEmitPrecision && bSupportsPrecisionQualifier) ? "highp " : "");
 						print_type_pre(var->type);
 						ralloc_asprintf_append(buffer, " %s", var->name);
