Fix "ViewMode ShaderComplexity" in Vulkan by adding approximate SPIRV instruction counts during shader compilation.

diff --git a/Engine/Source/Developer/VulkanShaderFormat/Private/VulkanShaderCompiler.cpp b/Engine/Source/Developer/VulkanShaderFormat/Private/VulkanShaderCompiler.cpp
--- a/Engine/Source/Developer/VulkanShaderFormat/Private/VulkanShaderCompiler.cpp
+++ b/Engine/Source/Developer/VulkanShaderFormat/Private/VulkanShaderCompiler.cpp
@@ -1871,6 +1871,55 @@ static void GatherSpirvReflectionBindings(
 	}
 }
 
+static uint32 CalculateSpirvInstructionCount(FVulkanSpirv& Spirv)
+{
+	// Count instructions inside functions
+	bool bInsideFunction = false;
+	uint32 ApproxInstructionCount = 0;
+	for (FSpirvConstIterator Iter = Spirv.cbegin(); Iter != Spirv.cend(); ++Iter)
+	{
+		switch (Iter.Opcode())
+		{
+
+		case SpvOpFunction:
+		{
+			check(!bInsideFunction);
+			bInsideFunction = true;
+		}
+		break;
+
+		case SpvOpFunctionEnd:
+		{
+			check(bInsideFunction);
+			bInsideFunction = false;
+		}
+		break;
+
+		case SpvOpLabel:
+		case SpvOpAccessChain:
+		case SpvOpSelectionMerge:
+		case SpvOpCompositeConstruct:
+		case SpvOpCompositeInsert:
+		case SpvOpCompositeExtract:
+			// Skip a few ops that show up often but don't result in much work on their own
+			break;
+
+		default:
+		{
+			if (bInsideFunction)
+			{
+				++ApproxInstructionCount;
+			}
+		}
+		break;
+
+		}
+	}
+	check(!bInsideFunction);
+
+	return ApproxInstructionCount;
+}
+
 static bool BuildShaderOutputFromSpirv(
 	CrossCompiler::FShaderConductorContext&	CompilerContext,
 	FVulkanSpirv&							Spirv,
@@ -2266,13 +2315,15 @@ static bool BuildShaderOutputFromSpirv(
 	PatchSpirvReflectionEntries(Spirv);
 	Spirv.EntryPointName = PatchSpirvEntryPointWithCRC(Spirv, Spirv.CRC);
 
+	const uint32 ApproxInstructionCount = CalculateSpirvInstructionCount(Spirv);
+
 	BuildShaderOutput(
 		Output,
 		Input,
 		TCHAR_TO_ANSI(*MetaData),
 		MetaData.Len(),
 		BindingTable,
-		/*NumLines*/0,
+		ApproxInstructionCount,
 		Spirv,
 		DebugName,
 		true // source contains meta data only
