UE5 - Remove some hlslcc vestiges from Vulkan

diff --git a/Engine/Source/Developer/VulkanShaderFormat/Private/VulkanBackend.cpp b/Engine/Source/Developer/VulkanShaderFormat/Private/VulkanBackend.cpp
--- a/Engine/Source/Developer/VulkanShaderFormat/Private/VulkanBackend.cpp
+++ b/Engine/Source/Developer/VulkanShaderFormat/Private/VulkanBackend.cpp
@@ -31,6119 +31,49 @@
 
 #include "VulkanBackend.h"
 #include "VulkanShaderFormat.h"
-#include "hlslcc.h"
-#include "hlslcc_private.h"
 #include "VulkanBackend.h"
-#include "compiler.h"
 #include "ShaderCompilerCommon.h"
 
 #include "VulkanCommon.h"
 
 #include "CrossCompilerCommon.h"
 
-PRAGMA_DISABLE_SHADOW_VARIABLE_WARNINGS
-#include "glsl_parser_extras.h"
-PRAGMA_ENABLE_SHADOW_VARIABLE_WARNINGS
-
-#include "hash_table.h"
-#include "ir_rvalue_visitor.h"
-#include "PackUniformBuffers.h"
-#include "IRDump.h"
-//@todo-rco: Remove STL!
-#include <sstream>
-#include <vector>
-
-#if !PLATFORM_WINDOWS
-#define _strdup strdup
-#endif
-
-const char* VULKAN_SUBPASS_FETCH[9] = 
-{
-	"VulkanSubpassDepthFetch",
-	"VulkanSubpassFetch0",
-	"VulkanSubpassFetch1",
-	"VulkanSubpassFetch2",
-	"VulkanSubpassFetch3",
-	"VulkanSubpassFetch4",
-	"VulkanSubpassFetch5",
-	"VulkanSubpassFetch6",
-	"VulkanSubpassFetch7"
-};
-const char* VULKAN_SUBPASS_FETCH_VAR[9] = 
-{
-	"GENERATED_SubpassDepthFetchAttachment",
-	"GENERATED_SubpassFetchAttachment0",
-	"GENERATED_SubpassFetchAttachment1",
-	"GENERATED_SubpassFetchAttachment2",
-	"GENERATED_SubpassFetchAttachment3",
-	"GENERATED_SubpassFetchAttachment4",
-	"GENERATED_SubpassFetchAttachment5",
-	"GENERATED_SubpassFetchAttachment6",
-	"GENERATED_SubpassFetchAttachment7"
-};
-const TCHAR* VULKAN_SUBPASS_FETCH_VAR_W[9] =
-{
-	TEXT("GENERATED_SubpassDepthFetchAttachment"),
-	TEXT("GENERATED_SubpassFetchAttachment0"),
-	TEXT("GENERATED_SubpassFetchAttachment1"),
-	TEXT("GENERATED_SubpassFetchAttachment2"),
-	TEXT("GENERATED_SubpassFetchAttachment3"),
-	TEXT("GENERATED_SubpassFetchAttachment4"),
-	TEXT("GENERATED_SubpassFetchAttachment5"),
-	TEXT("GENERATED_SubpassFetchAttachment6"),
-	TEXT("GENERATED_SubpassFetchAttachment7")
-};
-
-static inline std::string FixHlslName(const glsl_type* Type, bool bUseTextureInsteadOfSampler = false)
-{
-	check(Type->is_image() || Type->is_vector() || Type->is_numeric() || Type->is_void() || Type->is_sampler() || Type->is_scalar());
-	std::string Name = Type->name;
-	if (Type == glsl_type::half_type)
-	{
-		return "float";
-	}
-	else if (Type == glsl_type::half2_type)
-	{
-		return "vec2";
-	}
-	else if (Type == glsl_type::half3_type)
-	{
-		return "vec3";
-	}
-	else if (Type == glsl_type::half4_type)
-	{
-		return "vec4";
-	}
-	else if (Type == glsl_type::half2x2_type)
-	{
-		return "mat2";
-	}
-	else if (Type == glsl_type::half2x3_type)
-	{
-		return "mat2x3";
-	}
-	else if (Type == glsl_type::half2x4_type)
-	{
-		return "mat2x4";
-	}
-	else if (Type == glsl_type::half3x2_type)
-	{
-		return "mat3x2";
-	}
-	else if (Type == glsl_type::half3x3_type)
-	{
-		return "mat3";
-	}
-	else if (Type == glsl_type::half3x4_type)
-	{
-		return "mat3x4";
-	}
-	else if (Type == glsl_type::half4x2_type)
-	{
-		return "mat4x2";
-	}
-	else if (Type == glsl_type::half4x3_type)
-	{
-		return "mat4x3";
-	}
-	else if (Type == glsl_type::half4x4_type)
-	{
-		return "mat4";
-	}
-	else if (Type->is_sampler() && !Type->sampler_buffer)
-	{
-		if (bUseTextureInsteadOfSampler)
-		{
-			// if this assert fires, take a look at the calls to hash_table_insert(sampler_type, ...) at the top of glsl_type::get_templated_instance in glsl_types.cpp
-			//  if the last parameter of the "new glsl_type()" invocation for the current Type->name is nullptr, there's your problem.  You need to add a string there, and then handle it here.
-			// The point of this block is to replace things that say "uniform sampler pz0" with "uniform texture pz0", so that you can generate valid SPIR-V for sharing samplers across multiple textures.
-			check(Type->HlslName);
-			if (!FCStringAnsi::Strcmp(Type->HlslName, "texturecube"))
-			{
-				return "textureCube";
-			}
-			else if (!FCStringAnsi::Strcmp(Type->HlslName, "texture2d"))
-			{
-				return "texture2D";
-			}
-			else if (!FCStringAnsi::Strcmp(Type->HlslName, "texture3d"))
-			{
-				return "texture3D";
-			}
-			else if (!FCStringAnsi::Strcmp(Type->HlslName, "texturecubearray"))
-			{
-				return "textureCubeArray";
-			}
-			else if (!FCStringAnsi::Strcmp(Type->HlslName, "texture2darray"))
-			{
-				return "texture2DArray";
-			}
-			else if (!FCStringAnsi::Strcmp(Type->HlslName, "texture2dms"))
-			{
-				return "texture2DArray";
-			}
-			else if (!FCStringAnsi::Strcmp(Type->HlslName, "texture2dmsarray"))
-			{
-				return "texture2DMSArray";
-			}
-
-			return Type->HlslName;
-		}
-		else if (!FCStringAnsi::Strcmp(Type->name, "samplerExternalOES"))
-		{
-			return "sampler2D";
-		}
-	}
-
-	return Name;
-}
-
-static bool UsesUEIntrinsic(exec_list* Instructions, const char * UEIntrinsic)
-{
-	struct SFindUEIntrinsic : public ir_hierarchical_visitor
-	{
-		bool bFound;
-		const char * UEIntrinsic;
-		SFindUEIntrinsic(const char * InUEIntrinsic) : bFound(false), UEIntrinsic(InUEIntrinsic) {}
-
-		virtual ir_visitor_status visit_enter(ir_call* IR) override
-		{
-			if (IR->use_builtin && !strcmp(IR->callee_name(), UEIntrinsic))
-			{
-				bFound = true;
-				return visit_stop;
-			}
-
-			return visit_continue;
-		}
-	};
-
-	SFindUEIntrinsic Visitor(UEIntrinsic);
-	Visitor.run(Instructions);
-	return Visitor.bFound;
-}
-
-static bool GlslTypeSupportsPrecisionQualifier(const glsl_type& GlslType)
-{
-	switch (GlslType.base_type)
-	{
-	case GLSL_TYPE_UINT:
-	case GLSL_TYPE_INT:
-	case GLSL_TYPE_HALF:
-	case GLSL_TYPE_FLOAT:
-		return true;
-	case GLSL_TYPE_ARRAY:
-		if (const glsl_type* GlslArrayElementType = GlslType.fields.array)
-		{
-			return GlslTypeSupportsPrecisionQualifier(*GlslArrayElementType);
-		}
-		return false;
-	default:
-		return false;
-	}
-}
-
-static inline void Scanf(char* Dest, const char* Format, float* OutValue)
-{
-#if PLATFORM_WINDOWS
-	sscanf_s(Dest, Format, OutValue);
-#else
-	sscanf(Dest, Format, OutValue);
-#endif
-}
-
-/**
-* This table must match the ir_expression_operation enum.
-*/
-static const char * const GLSLExpressionTable[ir_opcode_count][4] =
-{
-	{ "(~", ")", "", "" }, // ir_unop_bit_not,
-	{ "not(", ")", "", "!" }, // ir_unop_logic_not,
-	{ "(-", ")", "", "" }, // ir_unop_neg,
-	{ "abs(", ")", "", "" }, // ir_unop_abs,
-	{ "sign(", ")", "", "" }, // ir_unop_sign,
-	{ "(1.0/(", "))", "", "" }, // ir_unop_rcp,
-	{ "inversesqrt(", ")", "", "" }, // ir_unop_rsq,
-	{ "sqrt(", ")", "", "" }, // ir_unop_sqrt,
-	{ "exp(", ")", "", "" }, // ir_unop_exp,      /**< Log base e on gentype */
-	{ "log(", ")", "", "" }, // ir_unop_log,	     /**< Natural log on gentype */
-	{ "exp2(", ")", "", "" }, // ir_unop_exp2,
-	{ "log2(", ")", "", "" }, // ir_unop_log2,
-	{ "int(", ")", "", "" }, // ir_unop_f2i,      /**< Float-to-integer conversion. */
-	{ "float(", ")", "", "" }, // ir_unop_i2f,      /**< Integer-to-float conversion. */
-	{ "bool(", ")", "", "" }, // ir_unop_f2b,      /**< Float-to-boolean conversion */
-	{ "float(", ")", "", "" }, // ir_unop_b2f,      /**< Boolean-to-float conversion */
-	{ "bool(", ")", "", "" }, // ir_unop_i2b,      /**< int-to-boolean conversion */
-	{ "int(", ")", "", "" }, // ir_unop_b2i,      /**< Boolean-to-int conversion */
-	{ "uint(", ")", "", "" }, // ir_unop_b2u,
-	{ "bool(", ")", "", "" }, // ir_unop_u2b,
-	{ "uint(", ")", "", "" }, // ir_unop_f2u,
-	{ "float(", ")", "", "" }, // ir_unop_u2f,      /**< Unsigned-to-float conversion. */
-	{ "uint(", ")", "", "" }, // ir_unop_i2u,      /**< Integer-to-unsigned conversion. */
-	{ "int(", ")", "", "" }, // ir_unop_u2i,      /**< Unsigned-to-integer conversion. */
-	{ "int(", ")", "", "" }, // ir_unop_h2i,
-	{ "float(", ")", "", "" }, // ir_unop_i2h,
-	{ "(", ")", "", "" }, // ir_unop_h2f,
-	{ "(", ")", "", "" }, // ir_unop_f2h,
-	{ "bool(", ")", "", "" }, // ir_unop_h2b,
-	{ "float(", ")", "", "" }, // ir_unop_b2h,
-	{ "uint(", ")", "", "" }, // ir_unop_h2u,
-	{ "uint(", ")", "", "" }, // ir_unop_u2h,
-	{ "transpose(", ")", "", "" }, // ir_unop_transpose
-	{ "any(", ")", "", "" }, // ir_unop_any,
-	{ "all(", ")", "", "" }, // ir_unop_all,
-
-	/**
-	* \name Unary floating-point rounding operations.
-	*/
-	/*@{*/
-	{ "trunc(", ")", "", "" }, // ir_unop_trunc,
-	{ "ceil(", ")", "", "" }, // ir_unop_ceil,
-	{ "floor(", ")", "", "" }, // ir_unop_floor,
-	{ "fract(", ")", "", "" }, // ir_unop_fract,
-	{ "round(", ")", "", "" }, // ir_unop_round,
-	/*@}*/
-
-	/**
-	* \name Trigonometric operations.
-	*/
-	/*@{*/
-	{ "sin(", ")", "", "" }, // ir_unop_sin,
-	{ "cos(", ")", "", "" }, // ir_unop_cos,
-	{ "tan(", ")", "", "" }, // ir_unop_tan,
-	{ "asin(", ")", "", "" }, // ir_unop_asin,
-	{ "acos(", ")", "", "" }, // ir_unop_acos,
-	{ "atan(", ")", "", "" }, // ir_unop_atan,
-	{ "sinh(", ")", "", "" }, // ir_unop_sinh,
-	{ "cosh(", ")", "", "" }, // ir_unop_cosh,
-	{ "tanh(", ")", "", "" }, // ir_unop_tanh,
-	/*@}*/
-
-	/**
-	* \name Normalize.
-	*/
-	/*@{*/
-	{ "normalize(", ")", "", "" }, // ir_unop_normalize,
-	/*@}*/
-
-	/**
-	* \name Partial derivatives.
-	*/
-	/*@{*/
-	{ "dFdx(", ")", "", "" }, // ir_unop_dFdx,
-	{ "dFdy(", ")", "", "" }, // ir_unop_dFdy,
-	{ "dfdx_fine(", ")", "", "" }, // ir_unop_dFdxFine,
-	{ "dfdy_fine(", ")", "", "" }, // ir_unop_dFdyFine,
-	{ "dfdx_coarse(", ")", "", "" }, // ir_unop_dFdxCoarse,
-	{ "dfdy_coarse(", ")", "", "" }, // ir_unop_dFdyCoarse,
-	/*@}*/
-
-	{ "isnan(", ")", "", "" }, // ir_unop_isnan,
-	{ "isinf(", ")", "", "" }, // ir_unop_isinf,
-
-	{ "floatBitsToUint(", ")", "", "" }, // ir_unop_fasu,
-	{ "floatBitsToInt(", ")", "", "" }, // ir_unop_fasi,
-	{ "intBitsToFloat(", ")", "", "" }, // ir_unop_iasf,
-	{ "uintBitsToFloat(", ")", "", "" }, // ir_unop_uasf,
-
-	{ "bitfieldReverse(", ")", "", "" }, // ir_unop_bitreverse,
-	{ "bitCount(", ")", "", "" }, // ir_unop_bitcount,
-	{ "findMSB(", ")", "", "" }, // ir_unop_msb,
-	{ "findLSB(", ")", "", "" }, // ir_unop_lsb,
-
-	/**
-	* \name Saturate.
-	*/
-	/*@{*/
-	{ "ERROR_NO_SATURATE_FUNCS(", ")", "", "" }, // ir_unop_saturate,
-	/*@}*/
-
-	{ "ERROR_NO_NOISE_FUNCS(", ")", "", "" }, // ir_unop_noise,
-
-	{ "(", "+", ")", "" }, // ir_binop_add,
-	{ "(", "-", ")", "" }, // ir_binop_sub,
-	{ "(", "*", ")", "" }, // ir_binop_mul,
-	{ "(", "/", ")", "" }, // ir_binop_div,
-
-	/**
-	* Takes one of two combinations of arguments:
-	*
-	* - mod(vecN, vecN)
-	* - mod(vecN, float)
-	*
-	* Does not take integer types.
-	*/
-	{ "mod(", ",", ")", "%" }, // ir_binop_mod,
-	{ "modf(", ",", ")", "" }, // ir_binop_modf,
-
-	{ "step(", ",", ")", "" }, // ir_binop_step,
-
-	/**
-	* \name Binary comparison operators which return a boolean vector.
-	* The type of both operands must be equal.
-	*/
-	/*@{*/
-	{ "lessThan(", ",", ")", "<" }, // ir_binop_less,
-	{ "greaterThan(", ",", ")", ">" }, // ir_binop_greater,
-	{ "lessThanEqual(", ",", ")", "<=" }, // ir_binop_lequal,
-	{ "greaterThanEqual(", ",", ")", ">=" }, // ir_binop_gequal,
-	{ "equal(", ",", ")", "==" }, // ir_binop_equal,
-	{ "notEqual(", ",", ")", "!=" }, // ir_binop_nequal,
-	/**
-	* Returns single boolean for whether all components of operands[0]
-	* equal the components of operands[1].
-	*/
-	{ "(", "==", ")", "" }, // ir_binop_all_equal,
-	/**
-	* Returns single boolean for whether any component of operands[0]
-	* is not equal to the corresponding component of operands[1].
-	*/
-	{ "(", "!=", ")", "" }, // ir_binop_any_nequal,
-	/*@}*/
-
-	/**
-	* \name Bit-wise binary operations.
-	*/
-	/*@{*/
-	{ "(", "<<", ")", "" }, // ir_binop_lshift,
-	{ "(", ">>", ")", "" }, // ir_binop_rshift,
-	{ "(", "&", ")", "" }, // ir_binop_bit_and,
-	{ "(", "^", ")", "" }, // ir_binop_bit_xor,
-	{ "(", "|", ")", "" }, // ir_binop_bit_or,
-	/*@}*/
-
-	{ "bvec%d(uvec%d(", ")*uvec%d(", "))", "&&" }, // ir_binop_logic_and,
-	{ "bvec%d(abs(ivec%d(", ")+ivec%d(", ")))", "^^" }, // ir_binop_logic_xor,
-	{ "bvec%d(uvec%d(", ")+uvec%d(", "))", "||" }, // ir_binop_logic_or,
-
-	{ "dot(", ",", ")", "" }, // ir_binop_dot,
-	{ "cross(", ",", ")", "" }, // ir_binop_cross,
-	{ "min(", ",", ")", "" }, // ir_binop_min,
-	{ "max(", ",", ")", "" }, // ir_binop_max,
-	{ "atan(", ",", ")", "" },
-	{ "pow(", ",", ")", "" }, // ir_binop_pow,
-
-	{ "mix(", ",", ",", ")" }, // ir_ternop_lerp,
-	{ "smoothstep(", ",", ",", ")" }, // ir_ternop_smoothstep,
-	{ "clamp(", ",", ",", ")" }, // ir_ternop_clamp,
-	{ "ERROR_NO_FMA_FUNCS(", ",", ",", ")" }, // ir_ternop_fma,
-
-	{ "ERROR_QUADOP_VECTOR(", ",", ")" }, // ir_quadop_vector,
-};
-
-static const char* OutputStreamTypeStrings[4] = {
-	"!invalid!",
-	"points",
-	"line_strip",
-	"triangle_strip"
-};
-
-static const char* GeometryInputStrings[6] = {
-	"!invalid!",
-	"points",
-	"lines",
-	"line_adjacency",
-	"triangles",
-	"triangles_adjacency"
-};
-
-static const char* DomainStrings[4] = {
-	"!invalid!",
-	"triangles",
-	"quads",
-	"isolines",
-};
-
-static const char* PartitioningStrings[5] = {
-	"!invalid!",
-	"equal_spacing",
-	"fractional_even_spacing",
-	"fractional_odd_spacing",
-	"pow2",
-};
-
-static const char* OutputTopologyStrings[5] = {
-	"!invalid!",
-	"point_needs_to_be_fixed",
-	"line_needs_to_be_fixed",
-	"cw",
-	"ccw",
-};
-
-static const char* GLSLIntCastTypes[5] =
-{
-	"!invalid!",
-	"int",
-	"ivec2",
-	"ivec3",
-	"ivec4",
-};
-
-static_assert((sizeof(GLSLExpressionTable) / sizeof(GLSLExpressionTable[0])) == ir_opcode_count, "GLSLExpressionTableSizeMismatch");
-
-// Holds information required for knowing if samplerstates should be shared
-struct FSamplerMappingGatherData
-{
-	struct FEntry
-	{
-		bool bUsingLoadOrDim = false; // either "Load" or "GetDimensions" intrinsics are used
-		bool bUsingSampling = false; // at least one texture intrinsic that requires a sampler state was used, e.g. "Sample", "SampleLevel" etc.
-		TStringSet SamplerStates;
-	};
-	std::map<std::string, FEntry> Entries;
-	TStringToSetMap SamplerToTextureMap;
-};
-
-struct FSamplerMapping
-{
-	// Final Data
-	TStringSet StandaloneSamplerStates;
-	TStringSet StandaloneTextures;
-	TStringSet CombinedSamplers;
-	bool bConsolidated = false;
-
-	void Consolidate(FSamplerMappingGatherData& GatherData)
-	{
-		check(!bConsolidated);
-
-		// First find all samplers using T.Load()
-		for (auto EntryPair : GatherData.Entries)
-		{
-			// Use a combined sampler if "Load" or "GetDimensions" intrinsic is used, except when sampling is used as well.
-			if (EntryPair.second.bUsingLoadOrDim && !EntryPair.second.bUsingSampling)
-			{
-				CombinedSamplers.insert(EntryPair.first);
-			}
-		}
-
-		// Now count how many shared sampler states exist
-		std::map<std::string, uint32> TexturesUsedPerSampler;
-		TStringSet CombinedSamplerStates;
-		for (auto Pair : GatherData.SamplerToTextureMap)
-		{
-			uint32 NumTexturesPerSampler = 0;
-			for (auto& Texture : Pair.second)
-			{
-				// Only count the textures that are *not* listed as combined-texture-samplers
-				if (CombinedSamplers.find(Texture) == CombinedSamplers.end())
-				{
-					++NumTexturesPerSampler;
-				}
-			}
-
-			if (NumTexturesPerSampler == 0 || NumTexturesPerSampler == 1)
-			{
-				CombinedSamplerStates.insert(Pair.first);
-			}
-
-			TexturesUsedPerSampler[Pair.first] = NumTexturesPerSampler;
-		}
-
-		// Now add combined samplers (ones with one samplerstate)
-		for (auto Pair : TexturesUsedPerSampler)
-		{
-			if (Pair.second == 1)
-			{
-				// Combined
-				auto Found = GatherData.SamplerToTextureMap.find(Pair.first);
-				if (Found != GatherData.SamplerToTextureMap.end())
-				{
-					for (auto& Texture : Found->second)
-					{
-						// Make sure this texture is not using multiple samplers
-						auto FoundTexture = GatherData.Entries.find(Texture);
-						check(FoundTexture != GatherData.Entries.end());
-						std::set<std::string>& SamplerStates = FoundTexture->second.SamplerStates;
-						uint32 NumSS = (uint32)SamplerStates.size();
-						if (NumSS <= 1)
-						{
-							CombinedSamplers.insert(Texture);
-						}
-					}
-				}
-			}
-		}
-
-		// Add all Textures and SamplerStates NOT in CombinedSamplers/CombinedSamplerStates as standalone
-		for (auto Pair : GatherData.SamplerToTextureMap)
-		{
-			uint32 NumStandaloneTexturesAdded = 0;
-			bool bIsSharedSamplerState = CombinedSamplerStates.find(Pair.first) == CombinedSamplerStates.end();
-			for (auto Texture : Pair.second)
-			{
-				if (CombinedSamplers.find(Texture) == CombinedSamplers.end())
-				{
-					StandaloneTextures.insert(Texture);
-					++NumStandaloneTexturesAdded;
-				}
-			}
-
-			if (bIsSharedSamplerState)
-			{
-				check(NumStandaloneTexturesAdded > 1);
-				StandaloneSamplerStates.insert(Pair.first);
-			}
-		}
-
-		bConsolidated = true;
-	}
-
-	bool UseCombinedImageSamplerForTexture(const char* TextureName)
-	{
-		check(bConsolidated);
-		return CombinedSamplers.find(TextureName) != CombinedSamplers.end();
-	}
-};
-
-struct SDMARange
-{
-	unsigned SourceCB;
-	unsigned SourceOffset;
-	unsigned Size;
-	unsigned DestCBIndex;
-	unsigned DestCBPrecision;
-	unsigned DestOffset;
-
-	bool operator <(SDMARange const & Other) const
-	{
-		if (SourceCB == Other.SourceCB)
-		{
-			return SourceOffset < Other.SourceOffset;
-		}
-
-		return SourceCB < Other.SourceCB;
-	}
-};
-typedef std::list<SDMARange> TDMARangeList;
-typedef std::map<unsigned, TDMARangeList> TCBDMARangeMap;
-
-
-static void InsertRange(TCBDMARangeMap& CBAllRanges, unsigned SourceCB, unsigned SourceOffset, unsigned Size, unsigned DestCBIndex, unsigned DestCBPrecision, unsigned DestOffset)
-{
-	check(SourceCB < (1 << 12));
-	check(DestCBIndex < (1 << 12));
-	check(DestCBPrecision < (1 << 8));
-	unsigned SourceDestCBKey = (SourceCB << 20) | (DestCBIndex << 8) | DestCBPrecision;
-	SDMARange Range = { SourceCB, SourceOffset, Size, DestCBIndex, DestCBPrecision, DestOffset };
-
-	TDMARangeList& CBRanges = CBAllRanges[SourceDestCBKey];
-//printf("* InsertRange: %08x\t%u:%u - %u:%c:%u:%u\n", SourceDestCBKey, SourceCB, SourceOffset, DestCBIndex, DestCBPrecision, DestOffset, Size);
-	if (CBRanges.empty())
-	{
-		CBRanges.push_back(Range);
-	}
-	else
-	{
-		TDMARangeList::iterator Prev = CBRanges.end();
-		bool bAdded = false;
-		for (auto Iter = CBRanges.begin(); Iter != CBRanges.end(); ++Iter)
-		{
-			if (SourceOffset + Size <= Iter->SourceOffset)
-			{
-				if (Prev == CBRanges.end())
-				{
-					CBRanges.push_front(Range);
-				}
-				else
-				{
-					CBRanges.insert(Iter, Range);
-				}
-
-				bAdded = true;
-				break;
-			}
-
-			Prev = Iter;
-		}
-
-		if (!bAdded)
-		{
-			CBRanges.push_back(Range);
-		}
-
-		if (CBRanges.size() > 1)
-		{
-			// Try to merge ranges
-			bool bDirty = false;
-			do
-			{
-				bDirty = false;
-				TDMARangeList NewCBRanges;
-				for (auto Iter = CBRanges.begin(); Iter != CBRanges.end(); ++Iter)
-				{
-					if (Iter == CBRanges.begin())
-					{
-						Prev = CBRanges.begin();
-					}
-					else
-					{
-						if (Prev->SourceOffset + Prev->Size == Iter->SourceOffset && Prev->DestOffset + Prev->Size == Iter->DestOffset)
-						{
-							SDMARange Merged = *Prev;
-							Merged.Size = Prev->Size + Iter->Size;
-							NewCBRanges.pop_back();
-							NewCBRanges.push_back(Merged);
-							++Iter;
-							NewCBRanges.insert(NewCBRanges.end(), Iter, CBRanges.end());
-							bDirty = true;
-							break;
-						}
-					}
-
-					NewCBRanges.push_back(*Iter);
-					Prev = Iter;
-				}
-
-				CBRanges.swap(NewCBRanges);
-			} while (bDirty);
-		}
-	}
-}
-
-static TDMARangeList SortRanges(TCBDMARangeMap& CBRanges)
-{
-	TDMARangeList Sorted;
-	for (auto& Pair : CBRanges)
-	{
-		Sorted.insert(Sorted.end(), Pair.second.begin(), Pair.second.end());
-	}
-
-	Sorted.sort();
-
-	return Sorted;
-}
-
-static void DumpSortedRanges(TDMARangeList& SortedRanges)
-{
-	printf("**********************************\n");
-	for (auto& o : SortedRanges)
-	{
-		printf("\t%u:%u - %u:%c:%u:%u\n", o.SourceCB, o.SourceOffset, o.DestCBIndex, o.DestCBPrecision, o.DestOffset, o.Size);
-	}
-}
-
-static inline ShaderStage::EStage GetDescriptorSetForStage(_mesa_glsl_parser_targets Target)
-{
-	switch (Target)
-	{
-	case vertex_shader:						return ShaderStage::GetStageForFrequency(SF_Vertex);
-	case fragment_shader:					return ShaderStage::GetStageForFrequency(SF_Pixel);
-	case compute_shader:					return ShaderStage::GetStageForFrequency(SF_Compute);
-	case geometry_shader:					return ShaderStage::GetStageForFrequency(SF_Geometry);
-	default: check(0);	break;	// NOT IMPLEMENTED!
-	}
-
-	return ShaderStage::Invalid;
-}
-
-/**
-* IR visitor used to generate GLSL. Based on ir_print_visitor.
-*/
-class FGenerateVulkanVisitor : public ir_visitor
-{
-	/** Track which multi-dimensional arrays are used. */
-	struct md_array_entry : public exec_node
-	{
-		const glsl_type* type;
-	};
-
-	/** Track external variables. */
-	struct extern_var : public exec_node
-	{
-		ir_variable* var;
-		explicit extern_var(ir_variable* in_var) : var(in_var) {}
-	};
-
-	/** External variables. */
-	exec_list input_variables;
-	exec_list output_variables;
-	exec_list uniform_variables;
-	exec_list sampler_variables;
-	exec_list image_variables;
-
-	/** Data tied globally to the shader via attributes */
-	bool early_depth_stencil;
-	int wg_size_x;
-	int wg_size_y;
-	int wg_size_z;
-
-	std::vector<std::string> ExternalSamplersList;
-
-	glsl_tessellation_info tessellation;
-
-	/** Track global instructions. */
-	struct global_ir : public exec_node
-	{
-		ir_instruction* ir;
-		explicit global_ir(ir_instruction* in_ir) : ir(in_ir) {}
-	};
-
-	/** Global instructions. */
-	exec_list global_instructions;
-
-	/** A mapping from ir_variable * -> unique printable names. */
-	hash_table *printable_names;
-	/** Structures required by the code. */
-	hash_table *used_structures;
-	/** Uniform block variables required by the code. */
-	hash_table *used_uniform_blocks;
-	/** Multi-dimensional arrays required by the code. */
-	exec_list used_md_arrays;
-
-	// Code generation flags
-	bool bIsES;
-	bool bEmitPrecision;
-	bool bIsES31;
-	EHlslCompileTarget Target;
-	_mesa_glsl_parse_state* ParseState;
-	bool bGenerateLayoutLocations;
-	bool bDefaultPrecisionIsHalf;
-
-	FVulkanBindingTable& BindingTable;
-
-	/** Memory context within which to make allocations. */
-	void *mem_ctx;
-	/** Buffer to which GLSL source is being generated. */
-	char** buffer;
-	/** Indentation level. */
-	int indentation;
-	/** Scope depth. */
-	int scope_depth;
-	/** The number of temporary variables declared in the current scope. */
-	int temp_id;
-	/** The number of global variables declared. */
-	int global_id;
-	/** Whether a semicolon must be printed before the next EOL. */
-	bool needs_semicolon;
-	/**
-	* Whether uint literals should be printed as int literals. This is a hack
-	* because glCompileShader crashes on Mac OS X with code like this:
-	* foo = bar[0u];
-	*/
-	bool should_print_uint_literals_as_ints;
-	/** number of loops in the generated code */
-	int loop_count;
-
-	/** Whether the shader being cross compiled needs EXT_shader_texture_lod. */
-	bool bUsesGLTextureLODExtension;
-
-	/** Found dFdx or dFdy */
-	bool bUsesDXDY;
-
-	// True if the discard instruction was encountered.
-	bool bUsesDiscard;
-
-	/** Found image atomic functions (e.g. imageAtomicAdd) */
-	bool bUsesImageWriteAtomic;
-
-	std::vector<std::string> SamplerStateNames;
-	TIRVarSet AtomicVariables;
-
-	/**
-	* Return true if the type is a multi-dimensional array. Also, track the
-	* array.
-	*/
-	bool is_md_array(const glsl_type* type)
-	{
-		if (type->base_type == GLSL_TYPE_ARRAY &&
-			type->fields.array->base_type == GLSL_TYPE_ARRAY)
-		{
-			foreach_iter(exec_list_iterator, iter, used_md_arrays)
-			{
-				md_array_entry* entry = (md_array_entry*)iter.get();
-				if (entry->type == type)
-					return true;
-			}
-			md_array_entry* entry = new(mem_ctx)md_array_entry();
-			entry->type = type;
-			used_md_arrays.push_tail(entry);
-			return true;
-		}
-		return false;
-	}
-
-	/**
-	* Fetch/generate a unique name for ir_variable.
-	*
-	* GLSL IR permits multiple ir_variables to share the same name.  This works
-	* fine until we try to print it, when we really need a unique one.
-	*/
-	const char *unique_name(ir_variable *var)
-	{
-		if (var->mode == ir_var_temporary || var->mode == ir_var_auto)
-		{
-			/* Do we already have a name for this variable? */
-			const char *name = (const char *)hash_table_find(this->printable_names, var);
-			if (name == NULL)
-			{
-				bool bIsGlobal = (scope_depth == 0 && var->mode != ir_var_temporary);
-				const char* prefix = "g";
-				if (!bIsGlobal)
-				{
-					if (var->type->is_matrix())
-					{
-						prefix = "m";
-					}
-					else if (var->type->is_vector())
-					{
-						prefix = "v";
-					}
-					else
-					{
-						switch (var->type->base_type)
-						{
-						case GLSL_TYPE_BOOL: prefix = "b"; break;
-						case GLSL_TYPE_UINT: prefix = "u"; break;
-						case GLSL_TYPE_INT: prefix = "i"; break;
-						case GLSL_TYPE_HALF: prefix = "h"; break;
-						case GLSL_TYPE_FLOAT: prefix = "f"; break;
-						default: prefix = "t"; break;
-						}
-					}
-				}
-				int var_id = bIsGlobal ? global_id++ : temp_id++;
-				name = ralloc_asprintf(mem_ctx, "%s%d", prefix, var_id);
-				hash_table_insert(this->printable_names, (void *)name, var);
-			}
-			return name;
-		}
-
-		/* If there's no conflict, just use the original name */
-		return var->name;
-	}
-
-	/**
-	* Add tabs/spaces for the current indentation level.
-	*/
-	void indent(void)
-	{
-		for (int i = 0; i < indentation; i++)
-		{
-			ralloc_asprintf_append(buffer, "\t");
-		}
-	}
-
-	/**
-	* Print out the internal name for a multi-dimensional array.
-	*/
-	void print_md_array_type(const glsl_type *t)
-	{
-		if (t->base_type == GLSL_TYPE_ARRAY)
-		{
-			ralloc_asprintf_append(buffer, "_mdarr_");
-			do
-			{
-				ralloc_asprintf_append(buffer, "%u_", t->length);
-				t = t->fields.array;
-			} while (t->base_type == GLSL_TYPE_ARRAY);
-			print_base_type(t);
-		}
-	}
-
-	/**
-	* Print the base type, e.g. vec3.
-	*/
-	void print_base_type(const glsl_type *t)
-	{
-		if (t->base_type == GLSL_TYPE_ARRAY)
-		{
-			print_base_type(t->fields.array);
-		}
-		else if (t->base_type == GLSL_TYPE_INPUTPATCH)
-		{
-			ralloc_asprintf_append(buffer, "/* %s */ ", t->name);
-			print_base_type(t->inner_type);
-		}
-		else if (t->base_type == GLSL_TYPE_OUTPUTPATCH)
-		{
-			ralloc_asprintf_append(buffer, "/* %s */ ", t->name);
-			print_base_type(t->inner_type);
-		}
-		else if ((t->base_type == GLSL_TYPE_STRUCT)
-			&& (FCStringAnsi::Strncmp("gl_", t->name, 3) != 0))
-		{
-			ralloc_asprintf_append(buffer, "%s", t->name);
-		}
-		else if (t->base_type == GLSL_TYPE_IMAGE)
-		{
-			ralloc_asprintf_append(buffer, "%s", t->name);
-		}
-		else
-		{
-			std::string Name = FixHlslName(t, ParseState->LanguageSpec->AllowsSharingSamplers());
-			ralloc_asprintf_append(buffer, "%s", Name.c_str());
-		}
-	}
-
-	/**
-	* Print the portion of the type that appears before a variable declaration.
-	*/
-	void print_type_pre(const glsl_type *t)
-	{
-		if (is_md_array(t))
-		{
-			print_md_array_type(t);
-		}
-		else
-		{
-			print_base_type(t);
-		}
-	}
-
-	/**
-	* Print the portion of the type that appears after a variable declaration.
-	*/
-	void print_type_post(const glsl_type *t, bool is_unsized = false)
-	{
-		if (t->base_type == GLSL_TYPE_ARRAY && !is_md_array(t))
-		{
-			if (is_unsized)
-			{
-				ralloc_asprintf_append(buffer, "[]");
-
-			}
-			else
-			{
-				ralloc_asprintf_append(buffer, "[%u]", t->length);
-			}
-		}
-		else if (t->base_type == GLSL_TYPE_INPUTPATCH || t->base_type == GLSL_TYPE_OUTPUTPATCH)
-		{
-			ralloc_asprintf_append(buffer, "[%u] /* %s */", t->patch_length, t->name);
-		}
-	}
-
-	/**
-	* Print a full variable declaration.
-	*/
-	void print_type_full(const glsl_type *t)
-	{
-		print_type_pre(t);
-		print_type_post(t);
-	}
-
-	/**
-	* Visit a single instruction. Appends a semicolon and EOL if needed.
-	*/
-	void do_visit(ir_instruction* ir)
-	{
-		needs_semicolon = true;
-		ir->accept(this);
-		if (needs_semicolon)
-		{
-			ralloc_asprintf_append(buffer, ";\n");
-		}
-	}
-
-	enum EPrecisionModifier
-	{
-		GLSL_PRECISION_DEFAULT,
-		GLSL_PRECISION_LOWP,
-		GLSL_PRECISION_MEDIUMP,
-		GLSL_PRECISION_HIGHP,
-	};
-
-	EPrecisionModifier GetPrecisionModifier(const struct glsl_type *type)
-	{
-		if (type->base_type == GLSL_TYPE_BOOL)
-		{
-			return GLSL_PRECISION_DEFAULT;
-		}
-		if (type->is_sampler() || type->is_image())
-		{
-			if (bDefaultPrecisionIsHalf && type->inner_type->base_type == GLSL_TYPE_FLOAT)
-			{
-				return GLSL_PRECISION_HIGHP;
-			}
-			else if (!bDefaultPrecisionIsHalf && type->inner_type->base_type == GLSL_TYPE_HALF)
-			{
-				return GLSL_PRECISION_MEDIUMP;
-			}
-			else // shadow samplers, integer textures etc
-			{
-				return GLSL_PRECISION_HIGHP;
-			}
-		}
-		else if (bDefaultPrecisionIsHalf && (type->base_type == GLSL_TYPE_FLOAT || (type->is_array() && type->element_type()->base_type == GLSL_TYPE_FLOAT)))
-		{
-			return GLSL_PRECISION_HIGHP;
-		}
-		else if (!bDefaultPrecisionIsHalf && (type->base_type == GLSL_TYPE_HALF || (type->is_array() && type->element_type()->base_type == GLSL_TYPE_HALF)))
-		{
-			return GLSL_PRECISION_MEDIUMP;
-		}
-		else if (type->is_integer())
-		{
-			return GLSL_PRECISION_HIGHP;
-		}
-		return GLSL_PRECISION_DEFAULT;
-	}
-
-	const char* GetPrecisionModifierName(EPrecisionModifier PrecisionModifier)
-	{
-		switch (PrecisionModifier)
-		{
-		case GLSL_PRECISION_LOWP:
-			return "lowp";
-			break;
-		case GLSL_PRECISION_MEDIUMP:
-			return "mediump";
-			break;
-		case GLSL_PRECISION_HIGHP:
-			return "highp";
-			break;
-		case GLSL_PRECISION_DEFAULT:
-			break;
-		default:
-			// we missed a type
-			check(false);
-		}
-		return "";
-	}
-
-	inline void AppendPrecisionModifier(char** inBuffer, EPrecisionModifier PrecisionModifier)
-	{
-		ralloc_asprintf_append(inBuffer, "%s ", GetPrecisionModifierName(PrecisionModifier));
-	}
-
-	/**
-	* \name Visit methods
-	*
-	* As typical for the visitor pattern, there must be one \c visit method for
-	* each concrete subclass of \c ir_instruction.  Virtual base classes within
-	* the hierarchy should not have \c visit methods.
-	*/
-
-	virtual void visit(ir_rvalue *rvalue)
-	{
-		check(0 && "ir_rvalue not handled for GLSL export.");
-	}
-
-	static inline bool is_struct_type(const glsl_type * type)
-	{
-		if (type->base_type != GLSL_TYPE_STRUCT && type->base_type != GLSL_TYPE_INPUTPATCH)
-		{
-			if (type->base_type == GLSL_TYPE_ARRAY && type->element_type())
-			{
-				return is_struct_type(type->element_type());
-			}
-			else
-			{
-				return false;
-			}
-		}
-		else
-		{
-			return true;
-		}
-	}
-
-	void print_zero_initialiser(const glsl_type * type)
-	{
-		if (type->base_type != GLSL_TYPE_STRUCT)
-		{
-			if (type->base_type != GLSL_TYPE_ARRAY)
-			{
-				ir_constant* zero = ir_constant::zero(mem_ctx, type);
-				if (zero)
-				{
-					zero->accept(this);
-				}
-			}
-			else
-			{
-				ralloc_asprintf_append(buffer, "{");
-
-				for (uint32 i = 0; i < type->length; i++)
-				{
-					if (i > 0)
-					{
-						ralloc_asprintf_append(buffer, ", ");
-					}
-					print_zero_initialiser(type->element_type());
-				}
-
-				ralloc_asprintf_append(buffer, "}");
-			}
-		}
-	}
-
-	virtual void visit(ir_variable *var)
-	{
-		const char * const centroid_str[] = { "", "centroid " };
-		const char * const invariant_str[] = { "", "invariant " };
-		const char * const patch_constant_str[] = { "", "patch " };
-		const char * const GLSLmode_str[] = { "", "uniform ", "in ", "out ", "inout ", "in ", "", "shared ", "", "", "uniform_ref " };
-		const char * const ESVSmode_str[] = { "", "uniform ", "attribute ", "varying ", "inout ", "in ", "", "shared " };
-		const char * const ESFSmode_str[] = { "", "uniform ", "varying ", "attribute ", "", "in ", "", "shared " };
-		const char * const GLSLinterp_str[] = { "", "smooth ", "flat ", "noperspective " };
-		const char * const ES31interp_str[] = { "", "smooth ", "flat ", "" };
-		const char * const layout_str[] = { "", "layout(origin_upper_left) ", "layout(pixel_center_integer) ", "layout(origin_upper_left,pixel_center_integer) " };
-
-		const char * const * mode_str = bIsES ? ((ParseState->target == vertex_shader) ? ESVSmode_str : ESFSmode_str) : GLSLmode_str;
-		const char * const * interp_str = bIsES31 ? ES31interp_str : GLSLinterp_str;
-
-		// Check for an initialized const variable
-		// If var is read-only and initialized, set it up as an initialized const
-		bool constInit = false;
-		if (var->has_initializer && var->read_only && (var->constant_initializer || var->constant_value))
-		{
-			ralloc_asprintf_append(buffer, "const ");
-			constInit = true;
-		}
-
-		CA_ASSUME(var->name);
-		CA_ASSUME(var->type);
-		CA_ASSUME(var->type->name);
-
-		if (scope_depth == 0)
-		{
-			glsl_base_type base_type = var->type->base_type;
-			if (base_type == GLSL_TYPE_ARRAY)
-			{
-				base_type = var->type->fields.array->base_type;
-			}
-
-			if (var->mode == ir_var_in)
-			{
-				input_variables.push_tail(new(mem_ctx)extern_var(var));
-			}
-			else if (var->mode == ir_var_out)
-			{
-				output_variables.push_tail(new(mem_ctx)extern_var(var));
-			}
-			else if (var->mode == ir_var_uniform && var->type->is_sampler())
-			{
-				sampler_variables.push_tail(new(mem_ctx)extern_var(var));
-			}
-			else if (var->mode == ir_var_uniform && var->type->is_image())
-			{
-				image_variables.push_tail(new(mem_ctx)extern_var(var));
-			}
-			else if (var->mode == ir_var_uniform && base_type == GLSL_TYPE_SAMPLER_STATE)
-			{
-				// ignore sampler state uniforms
-			}
-			else if (var->mode == ir_var_uniform && var->semantic == NULL)
-			{
-				uniform_variables.push_tail(new(mem_ctx)extern_var(var));
-			}
-		}
-
-		if (var->name && FCStringAnsi::Strncmp(var->name, "gl_", 3) == 0 &&
-			var->centroid == 0 && (var->interpolation == ir_interp_qualifier_none || var->interpolation == ir_interp_qualifier_flat) &&
-			var->invariant == 0 && var->origin_upper_left == 0 &&
-			var->pixel_center_integer == 0)
-		{
-			// Don't emit builtin GL variable declarations.
-			needs_semicolon = false;
-		}
-		else if (scope_depth == 0 && var->mode == ir_var_temporary)
-		{
-			global_instructions.push_tail(new(mem_ctx)global_ir(var));
-			needs_semicolon = false;
-		}
-		else
-		{
-			int layout_bits =
-				(var->origin_upper_left ? 0x1 : 0) |
-				(var->pixel_center_integer ? 0x2 : 0);
-
-			if (scope_depth == 0 &&
-				((var->mode == ir_var_in) || (var->mode == ir_var_out)) &&
-				var->is_interface_block)
-			{
-				/**
-				Hack to display our fake structs as what they are supposed to be - interface blocks
-
-				'in' or 'out' variable qualifier becomes interface block declaration start,
-				structure name becomes block name,
-				we add information about block contents, taking type from sole struct member type, and
-				struct variable name becomes block instance name.
-
-				Note: With tessellation, matching interfaces between shaders is tricky, so we need
-				to assign explicit locations to shader input and output variables.
-
-				The reason we use a struct instead of an interface block is that with
-				GL4.2/GL_ARB_separate_shader_objects, you can add a layout(location=foo) to a variable
-				that is not part of an interface block. However, in order to add a location to a variable
-				inside an interface block, you need GL4.4/GL_enhanced_layouts. Since for now, we don't want
-				that dependency, we use structs.
-
-				*/
-
-				if (bGenerateLayoutLocations && var->explicit_location && var->is_patch_constant == 0)
-				{
-					check(layout_bits == 0);
-					const glsl_type* inner_type = var->type;
-					if (inner_type->is_array())
-					{
-						inner_type = inner_type->fields.array;
-					}
-					check(inner_type->is_record());
-					check(inner_type->length == 1);
-					const glsl_struct_field* field = &inner_type->fields.structure[0];
-
-					ralloc_asprintf_append(
-						buffer,
-						"layout(location=%d) %s",
-						var->location,				// location number
-						mode_str[var->mode]			// in / out
-						);
-
-					// Append type to the buffer string
-					print_type_pre(field->type); // float, vec2, vec3 and etc..
-				}
-				else
-				{
-					ralloc_asprintf_append(
-						buffer,
-						"%s%s%s%s",
-						//layout_str[layout_bits],
-						centroid_str[var->centroid],
-						invariant_str[var->invariant],
-						patch_constant_str[var->is_patch_constant],
-						mode_str[var->mode]
-						);
-
-					print_type_pre(var->type);
-
-					const glsl_type* inner_type = var->type;
-					if (inner_type->is_array())
-					{
-						inner_type = inner_type->fields.array;
-					}
-					check(inner_type->is_record());
-					check(inner_type->length == 1);
-					const glsl_struct_field* field = &inner_type->fields.structure[0];
-					check(FCStringAnsi::Strcmp(field->name, "Data") == 0);
-
-					ralloc_asprintf_append(buffer, " { %s", interp_str[var->interpolation]);
-
-					print_type_pre(field->type);
-					ralloc_asprintf_append(buffer, " Data");
-					print_type_post(field->type);
-					ralloc_asprintf_append(buffer, "; }");
-				}
-			}
-			else if (var->type->is_image())
-			{
-				if (var->type->HlslName && (!strncmp(var->type->HlslName, "RWStructuredBuffer<", 19) || !strncmp(var->type->HlslName, "StructuredBuffer<", 17)))
-				{
-					const char * const readonly_str[] = { "", "readonly " };
-					const int readonly = strncmp(var->type->HlslName, "StructuredBuffer<", 17)==0 ? 1 : 0;
-										
-					ralloc_asprintf_append(
-						buffer,
-						"layout(set=%d,binding=BINDING_%d) %s buffer ",
-						GetDescriptorSetForStage(ParseState->target),
-						BindingTable.RegisterBinding(var->name, "u", EVulkanBindingType::StorageBuffer),
-						readonly_str[readonly]
-					);
-				}
-				else
-				{
-					const bool bSingleComp = (var->type->inner_type->vector_elements == 1);
-					const char * const coherent_str[] ={"", "coherent "};
-					const char * const writeonly_str[] ={"", "writeonly "};
-					const char * const type_str[] ={"32ui", "32i", "16f", "32f"};
-					const char * const comp_str = bSingleComp ? "r" : "rgba";
-					const int writeonly = var->image_write && !(var->image_read);
-
-					check(var->type->inner_type->base_type >= GLSL_TYPE_UINT && var->type->inner_type->base_type <= GLSL_TYPE_FLOAT);
-
-					ralloc_asprintf_append(
-						buffer,
-						"%s%s%s%s",
-						invariant_str[var->invariant],
-						mode_str[var->mode],
-						coherent_str[var->coherent],
-						writeonly_str[writeonly]
-					);
-
-					//should check here on base type
-					ralloc_asprintf_append(
-						buffer,
-						"layout(set=%d,%s%s,binding=BINDING_%d) ",
-						GetDescriptorSetForStage(ParseState->target),
-						comp_str,
-						type_str[var->type->inner_type->base_type],
-						BindingTable.RegisterBinding(var->name, "u", var->type->sampler_buffer ? EVulkanBindingType::StorageTexelBuffer : EVulkanBindingType::StorageImage)
-						);
-
-					if (bEmitPrecision)
-					{
-						AppendPrecisionModifier(buffer, GetPrecisionModifier(var->type));
-					}
-					print_type_pre(var->type);
-				}
-			}
-			else
-			{
-				char* layout = nullptr;
-
-				uint32 Interpolation = var->interpolation;
-				if (var->type->is_sampler())
-				{
-					EVulkanBindingType::EType BindingType = var->type->sampler_buffer
-						? EVulkanBindingType::UniformTexelBuffer
-						: (SamplerMapping.StandaloneTextures.find(std::string(var->name)) != SamplerMapping.StandaloneTextures.end()
-							? EVulkanBindingType::Image
-							: EVulkanBindingType::CombinedImageSampler);
-					int32 Binding = BindingTable.RegisterBinding(var->name, "s", BindingType);
-					layout = ralloc_asprintf(nullptr,
-						"layout(set=%d, binding=BINDING_%d) ",
-						GetDescriptorSetForStage(ParseState->target),
-						Binding);
-					if (var->type->name && !strcmp(var->type->name, "samplerExternalOES"))
-					{
-						ExternalSamplersList.push_back(var->name);
-					}
-				}
-				else if (bGenerateLayoutLocations && var->explicit_location)
-				{
-					check(layout_bits == 0);
-					layout = ralloc_asprintf(nullptr, "layout(location=%d) ", var->location);
-					if (ParseState->target == fragment_shader && var->type->is_integer() && var->mode == ir_var_in)
-					{
-						// Flat
-						Interpolation = 2;
-					}
-				}
-
-				ralloc_asprintf_append(
-					buffer,
-					"%s%s%s%s%s%s",
-					layout ? layout : layout_str[layout_bits],
-					centroid_str[var->centroid],
-					invariant_str[var->invariant],
-					patch_constant_str[var->is_patch_constant],
-					mode_str[var->mode],
-					interp_str[Interpolation]
-					);
-
-				if (bEmitPrecision)
-				{
-					AppendPrecisionModifier(buffer, GetPrecisionModifier(var->type));
-				}
-
-				if (bGenerateLayoutLocations && var->explicit_location)
-				{
-					ralloc_free(layout);
-				}
-
-				if (var->type->is_sampler() && ParseState->LanguageSpec->AllowsSharingSamplers())
-				{
-					if (SamplerMapping.UseCombinedImageSamplerForTexture(var->name))
-					{
-						std::string Name = FixHlslName(var->type, false);
-						ralloc_asprintf_append(buffer, "%s", Name.c_str());
-					}
-					else
-					{
-						print_type_pre(var->type);
-					}
-				}
-				else
-				{
-					print_type_pre(var->type);
-				}
-			}
-
-			if (var->type->is_image() && (var->type->HlslName && (!strncmp(var->type->HlslName, "RWStructuredBuffer<", 19) || !strncmp(var->type->HlslName, "StructuredBuffer<", 17))))
-			{
-				AddTypeToUsedStructs(var->type->inner_type);
-				// DO NOT change _BUFFER (or update when reading the spir-v reflection)
-				ralloc_asprintf_append(buffer, " %s_BUFFER { %s %s[]; }", unique_name(var), var->type->inner_type->name, unique_name(var));
-			}
-			else
-			{
-				ralloc_asprintf_append(buffer, " %s", unique_name(var));
-				const bool bUnsizedArray = var->mode == ir_var_in && ((ParseState->target == tessellation_evaluation_shader) || (ParseState->target == tessellation_control_shader));
-				print_type_post(var->type, bUnsizedArray);
-			}
-		}
-
-		// Add the initializer if we need it
-		if (constInit)
-		{
-			ralloc_asprintf_append(buffer, " = ");
-			if (var->constant_initializer)
-			{
-				var->constant_initializer->accept(this);
-			}
-			else
-			{
-				var->constant_value->accept(this);
-			}
-		}
-
-		// add type to used_structures so we can later declare them at the start of the GLSL shader
-		// this is for the case of a variable that is declared, but not later dereferenced (which can happen
-		// when debugging HLSLCC and running without optimization
-		AddTypeToUsedStructs(var->type);
-	}
-
-	virtual void visit(ir_function_signature *sig)
-	{
-		// Reset temporary id count.
-		temp_id = 0;
-		bool bPrintComma = false;
-		scope_depth++;
-
-		print_type_full(sig->return_type);
-		ralloc_asprintf_append(buffer, " %s(", sig->function_name());
-
-		foreach_iter(exec_list_iterator, iter, sig->parameters)
-		{
-			ir_variable *const inst = (ir_variable *)iter.get();
-			if (bPrintComma)
-			{
-				ralloc_asprintf_append(buffer, ",");
-			}
-			inst->accept(this);
-			bPrintComma = true;
-		}
-		ralloc_asprintf_append(buffer, ")\n");
-
-		indent();
-		ralloc_asprintf_append(buffer, "{\n");
-
-		if (sig->is_main && !global_instructions.is_empty())
-		{
-			indentation++;
-			foreach_iter(exec_list_iterator, iter, global_instructions)
-			{
-				global_ir* gir = (global_ir*)iter.get();
-				indent();
-				do_visit(gir->ir);
-			}
-			indentation--;
-		}
-
-		//grab the global attributes
-		if (sig->is_main)
-		{
-			early_depth_stencil = sig->is_early_depth_stencil;
-			wg_size_x = sig->wg_size_x;
-			wg_size_y = sig->wg_size_y;
-			wg_size_z = sig->wg_size_z;
-
-			tessellation = sig->tessellation;
-		}
-
-		indentation++;
-		foreach_iter(exec_list_iterator, iter, sig->body)
-		{
-			ir_instruction *const inst = (ir_instruction *)iter.get();
-			indent();
-			do_visit(inst);
-		}
-		indentation--;
-		indent();
-		ralloc_asprintf_append(buffer, "}\n");
-		needs_semicolon = false;
-		scope_depth--;
-	}
-
-	virtual void visit(ir_function *func)
-	{
-		foreach_iter(exec_list_iterator, iter, *func)
-		{
-			ir_function_signature *const sig = (ir_function_signature *)iter.get();
-			if (sig->is_defined && !sig->is_builtin)
-			{
-				indent();
-				sig->accept(this);
-			}
-		}
-		needs_semicolon = false;
-	}
-
-	virtual void visit(ir_expression *expr)
-	{
-		check(scope_depth > 0);
-
-		int numOps = expr->get_num_operands();
-		ir_expression_operation op = expr->operation;
-
-		if (numOps == 1 && op >= ir_unop_first_conversion && op <= ir_unop_last_conversion)
-		{
-			if (op == ir_unop_f2h || op == ir_unop_h2f)
-			{
-				// No need to convert from half<->float as that is part of the precision of a variable
-				expr->operands[0]->accept(this);
-			}
-			else
-			{
-				ralloc_asprintf_append(buffer, "%s(", FixHlslName(expr->type).c_str());
-				expr->operands[0]->accept(this);
-				ralloc_asprintf_append(buffer, ")");
-			}
-		}
-		else if (expr->type->is_scalar() &&
-			((numOps == 1 && op == ir_unop_logic_not) ||
-			(numOps == 2 && op >= ir_binop_first_comparison && op <= ir_binop_last_comparison) ||
-			(numOps == 2 && op >= ir_binop_first_logic && op <= ir_binop_last_logic)))
-		{
-			const char* op_str = GLSLExpressionTable[op][3];
-			ralloc_asprintf_append(buffer, "%s(", (numOps == 1) ? op_str : "");
-			expr->operands[0]->accept(this);
-			if (numOps == 2)
-			{
-				ralloc_asprintf_append(buffer, "%s", op_str);
-				expr->operands[1]->accept(this);
-			}
-			ralloc_asprintf_append(buffer, ")");
-		}
-		else if (expr->type->is_vector() && numOps == 2 &&
-			op >= ir_binop_first_logic && op <= ir_binop_last_logic)
-		{
-			ralloc_asprintf_append(buffer, GLSLExpressionTable[op][0], expr->type->vector_elements, expr->type->vector_elements);
-			expr->operands[0]->accept(this);
-			ralloc_asprintf_append(buffer, GLSLExpressionTable[op][1], expr->type->vector_elements);
-			expr->operands[1]->accept(this);
-			ralloc_asprintf_append(buffer, GLSLExpressionTable[op][2]);
-		}
-		else if (op == ir_binop_mod && !expr->type->is_float())
-		{
-			ralloc_asprintf_append(buffer, "((");
-			expr->operands[0]->accept(this);
-			ralloc_asprintf_append(buffer, ")%%(");
-			expr->operands[1]->accept(this);
-			ralloc_asprintf_append(buffer, "))");
-		}
-		else if (op == ir_binop_mul && expr->type->is_matrix()
-			&& expr->operands[0]->type->is_matrix()
-			&& expr->operands[1]->type->is_matrix())
-		{
-			ralloc_asprintf_append(buffer, "matrixCompMult(");
-			expr->operands[0]->accept(this);
-			ralloc_asprintf_append(buffer, ",");
-			expr->operands[1]->accept(this);
-			ralloc_asprintf_append(buffer, ")");
-		}
-		else if (numOps < 4)
-		{
-			if (op == ir_unop_dFdx || op == ir_unop_dFdy)
-			{
-				bUsesDXDY = true;
-			}
-
-			ralloc_asprintf_append(buffer, GLSLExpressionTable[op][0]);
-			for (int i = 0; i < numOps; ++i)
-			{
-				expr->operands[i]->accept(this);
-				ralloc_asprintf_append(buffer, GLSLExpressionTable[op][i + 1]);
-			}
-		}
-	}
-
-	virtual void visit(ir_texture *tex)
-	{
-		check(scope_depth > 0);
-
-		const char * const fetch_str[] = { "texture", "texelFetch" };
-		const char * const Dim[] = { "", "2D", "3D", "Cube", "", "", "" };
-		static const char * const size_str[] = { "", "Size" };
-		static const char * const proj_str[] = { "", "Proj" };
-		static const char * const grad_str[] = { "", "Grad" };
-		static const char * const lod_str[] = { "", "Lod" };
-		static const char * const offset_str[] = { "", "Offset" };
-		static const char * const gather_str[] = { "", "Gather" };
-		static const char * const querymips_str[] = { "", "QueryLevels" };
-		static const char * const EXT_str[] = { "", "EXT" };
-		const bool cube_array = tex->sampler->type->sampler_dimensionality == GLSL_SAMPLER_DIM_CUBE &&
-			tex->sampler->type->sampler_array;
-
-		ir_texture_opcode op = tex->op;
-		if (op == ir_txl && tex->sampler->type->sampler_shadow && tex->sampler->type->sampler_dimensionality == GLSL_SAMPLER_DIM_CUBE)
-		{
-			// This very instruction is missing in OpenGL 3.2, so we need to change the sampling to instruction that exists in order for shader to compile
-			op = ir_tex;
-		}
-
-		bool bEmitEXT = false;
-
-		if (bIsES && op == ir_txl)
-		{
-			// See http://www.khronos.org/registry/gles/extensions/EXT/EXT_shader_texture_lod.txt
-			bUsesGLTextureLODExtension = true;
-			bEmitEXT = true;
-		}
-
-		// Emit texture function and sampler.
-		ralloc_asprintf_append(buffer, "%s%s%s%s%s%s%s%s%s%s(",
-			fetch_str[op == ir_txf],
-			bIsES ? Dim[tex->sampler->type->sampler_dimensionality] : "",
-			gather_str[op == ir_txg],
-			size_str[op == ir_txs],
-			querymips_str[op == ir_txm],
-			proj_str[tex->projector != 0],
-			grad_str[op == ir_txd],
-			lod_str[op == ir_txl],
-			offset_str[tex->offset != 0],
-			EXT_str[(int)bEmitEXT]
-			);
-
-		if (ParseState->LanguageSpec->AllowsSharingSamplers() && !tex->sampler->type->sampler_buffer && 
-			!SamplerMapping.UseCombinedImageSamplerForTexture(tex->sampler->variable_referenced()->name))
-		{
-			uint32 SSIndex = AddUniqueSamplerState(tex->SamplerStateName);
-			BindingTable.RegisterBinding(tex->SamplerStateName, "z", EVulkanBindingType::Sampler);
-
-			auto GetSamplerSuffix = [](int32 Dim)
-			{
-				switch (Dim)
-				{
-				case GLSL_SAMPLER_DIM_1D: return "1D";
-				case GLSL_SAMPLER_DIM_2D: return "2D";
-				case GLSL_SAMPLER_DIM_3D: return "3D";
-				case GLSL_SAMPLER_DIM_CUBE: return "Cube";
-					//case GLSL_SAMPLER_DIM_RECT: return "Rect";
-					//case GLSL_SAMPLER_DIM_BUF: return "Buf";
-				default: return "INVALID";
-				}
-			};
-
-			const bool bHasSamplerState = (tex->SamplerStateName && *tex->SamplerStateName);
-			ralloc_asprintf_append(buffer, "sampler%s%s(", GetSamplerSuffix(tex->sampler->type->sampler_dimensionality),
-					tex->sampler->type->sampler_array ? "Array" : "");
-			tex->sampler->accept(this);
-			if ((op == ir_txf || op == ir_txs || op == ir_txm) && !bHasSamplerState)
-			{
-				// No sampler state was provided in the intrinsic, so look up the texture-to-sampler map to find a sampler state.
-				// This is required when a texture is used for intrinsics that actually don't need a sampler-state such as "texelFetch",
-				// but the texture was also used for an intrinsic that might need a separate sampler, so we have a "uniform texture2D" instead of a "uniform sampler2D".
-				// To avoid the dependency to the GLSL extension "GL_EXT_samplerless_texture_functions", which provides "texelFetch(texture2D,...)",
-				// we construct a combined-sampler with *any* of the samplers for this texture such as "sampler2D(myTexture, mySampler)".
-				auto Found = ParseState->TextureToSamplerMap.find(tex->sampler->variable_referenced()->name);
-
-				// Can't find a sampler state for this texture, internal error!
-				std::string SamplerStateName = "INTERNAL_ERROR_MISSING_SAMPLERSTATE";
-
-				if (Found != ParseState->TextureToSamplerMap.end())
-				{
-					TStringSet& SamplerStates = Found->second;
-					for (auto& Name : SamplerStates)
-					{
-						if (Name != "")
-						{
-							SamplerStateName = Name;
-							break;
-						}
-					}
-				}
-				ralloc_asprintf_append(buffer, ", %s)", SamplerStateName.c_str());
-			}
-			else
-			{
-				ralloc_asprintf_append(buffer, ", %s)", tex->SamplerStateName);
-			}
-		}
-		else
-		{
-			tex->sampler->accept(this);
-		}
-
-		// Emit coordinates.
-		if ((op == ir_txs && tex->lod_info.lod) || op == ir_txm)
-		{
-			if (!tex->sampler->type->sampler_ms && op != ir_txm)
-			{
-				ralloc_asprintf_append(buffer, ",");
-				tex->lod_info.lod->accept(this);
-			}
-		}
-		else if (tex->sampler->type->sampler_shadow && (op != ir_txg && !cube_array))
-		{
-			int coord_dims = 0;
-			switch (tex->sampler->type->sampler_dimensionality)
-			{
-			case GLSL_SAMPLER_DIM_1D: coord_dims = 2; break;
-			case GLSL_SAMPLER_DIM_2D: coord_dims = 3; break;
-			case GLSL_SAMPLER_DIM_3D: coord_dims = 4; break;
-			case GLSL_SAMPLER_DIM_CUBE: coord_dims = 4; break;
-			default: check(0 && "Shadow sampler has unsupported dimensionality.");
-			}
-			ralloc_asprintf_append(buffer, ",vec%d(", coord_dims);
-			tex->coordinate->accept(this);
-			ralloc_asprintf_append(buffer, ",");
-			tex->shadow_comparitor->accept(this);
-			ralloc_asprintf_append(buffer, ")");
-		}
-		else
-		{
-			ralloc_asprintf_append(buffer, ",");
-			tex->coordinate->accept(this);
-		}
-
-		// Emit gather compare value
-		if (tex->sampler->type->sampler_shadow && (op == ir_txg || cube_array))
-		{
-			ralloc_asprintf_append(buffer, ",");
-			tex->shadow_comparitor->accept(this);
-		}
-
-		// Emit sample index.
-		if (op == ir_txf && tex->sampler->type->sampler_ms)
-		{
-			ralloc_asprintf_append(buffer, ",");
-			tex->lod_info.sample_index->accept(this);
-		}
-
-		// Emit LOD.
-		if (op == ir_txl ||
-			(op == ir_txf && tex->lod_info.lod &&
-			!tex->sampler->type->sampler_ms && !tex->sampler->type->sampler_buffer))
-		{
-			ralloc_asprintf_append(buffer, ",");
-			tex->lod_info.lod->accept(this);
-		}
-
-		// Emit gradients.
-		if (op == ir_txd)
-		{
-			ralloc_asprintf_append(buffer, ",");
-			tex->lod_info.grad.dPdx->accept(this);
-			ralloc_asprintf_append(buffer, ",");
-			tex->lod_info.grad.dPdy->accept(this);
-		}
-		else if (op == ir_txb)
-		{
-			ralloc_asprintf_append(buffer, ",");
-			tex->lod_info.bias->accept(this);
-		}
-
-		// Emit offset.
-		if (tex->offset)
-		{
-			ralloc_asprintf_append(buffer, ",");
-			tex->offset->accept(this);
-		}
-
-		// Emit channel selection for gather
-		if (op == ir_txg && tex->channel > ir_channel_none)
-		{
-			check(tex->channel < ir_channel_unknown);
-			ralloc_asprintf_append(buffer, ", %d", int(tex->channel) - 1);
-		}
-
-		ralloc_asprintf_append(buffer, ")");
-	}
-
-	virtual void visit(ir_swizzle *swizzle)
-	{
-		check(scope_depth > 0);
-
-		const unsigned mask[4] =
-		{
-			swizzle->mask.x,
-			swizzle->mask.y,
-			swizzle->mask.z,
-			swizzle->mask.w,
-		};
-
-		if (swizzle->val->type->is_scalar())
-		{
-			// Scalar -> Vector swizzles must use the constructor syntax.
-			if (swizzle->type->is_scalar() == false)
-			{
-				print_type_full(swizzle->type);
-				ralloc_asprintf_append(buffer, "(");
-				swizzle->val->accept(this);
-				ralloc_asprintf_append(buffer, ")");
-			}
-		}
-		else
-		{
-			const bool is_constant = swizzle->val->as_constant() != nullptr;
-			if (is_constant)
-			{
-				ralloc_asprintf_append(buffer, "(");
-			}
-			swizzle->val->accept(this);
-			if (is_constant)
-			{
-				ralloc_asprintf_append(buffer, ")");
-			}
-			ralloc_asprintf_append(buffer, ".");
-			for (unsigned i = 0; i < swizzle->mask.num_components; i++)
-			{
-				ralloc_asprintf_append(buffer, "%c", "xyzw"[mask[i]]);
-			}
-		}
-	}
-
-	virtual void visit(ir_dereference_variable *deref)
-	{
-		check(scope_depth > 0);
-
-		ir_variable* var = deref->variable_referenced();
-
-		ralloc_asprintf_append(buffer, unique_name(var));
-
-
-		// add type to used_structures so we can later declare them at the start of the GLSL shader
-		AddTypeToUsedStructs(var->type);
-
-
-		if (var->mode == ir_var_uniform && var->semantic != NULL)
-		{
-			if (hash_table_find(used_uniform_blocks, var->semantic) == NULL)
-			{
-				hash_table_insert(used_uniform_blocks, (void*)var->semantic, var->semantic);
-			}
-		}
-
-		if (is_md_array(deref->type))
-		{
-			ralloc_asprintf_append(buffer, ".Inner");
-		}
-	}
-
-	virtual void visit(ir_dereference_array *deref)
-	{
-		check(scope_depth > 0);
-
-		deref->array->accept(this);
-
-		// Make extra sure crappy Mac OS X compiler won't have any reason to crash
-		bool enforceInt = false;
-
-		if (deref->array_index->type->base_type == GLSL_TYPE_UINT)
-		{
-			if (deref->array_index->ir_type == ir_type_constant)
-			{
-				should_print_uint_literals_as_ints = true;
-			}
-			else
-			{
-				enforceInt = true;
-			}
-		}
-
-		if (enforceInt)
-		{
-			ralloc_asprintf_append(buffer, "[int(");
-		}
-		else
-		{
-			ralloc_asprintf_append(buffer, "[");
-		}
-
-		deref->array_index->accept(this);
-		should_print_uint_literals_as_ints = false;
-
-		if (enforceInt)
-		{
-			ralloc_asprintf_append(buffer, ")]");
-		}
-		else
-		{
-			ralloc_asprintf_append(buffer, "]");
-		}
-
-		if (is_md_array(deref->array->type))
-		{
-			ralloc_asprintf_append(buffer, ".Inner");
-		}
-	}
-
-	void print_image_op(ir_dereference_image *deref, ir_rvalue *src, const char* Mask = nullptr)
-	{
-		const char* swizzle[] =
-		{
-			"x", "xy", "xyz", "xyzw"
-		};
-		const char* expand[] =
-		{
-			"xxxx", "xyxx", "xyzx", "xyzw"
-		};
-		const int dst_elements = deref->type->vector_elements;
-		const int src_elements = (src) ? src->type->vector_elements : 1;
-
-		bool bIsStructured = deref->type->is_record() || (deref->image->type->HlslName && (!strncmp(deref->image->type->HlslName, "RWStructuredBuffer<", 19) || !strncmp(deref->image->type->HlslName, "StructuredBuffer<", 17)));
-
-		//!strncmp(var->type->name, "RWStructuredBuffer<")
-		check(bIsStructured || (1 <= dst_elements && dst_elements <= 4));
-		check(bIsStructured || (1 <= src_elements && src_elements <= 4));
-
-		if (deref->op == ir_image_access)
-		{
-			if (bIsStructured)
-			{
-				if (src == NULL)
-				{			
-					deref->image->accept(this);
-					ralloc_asprintf_append(buffer, "[");
-					deref->image_index->accept(this);
-					ralloc_asprintf_append(buffer, "]");					
-				}
-				else
-				{
-					deref->image->accept(this);
-					ralloc_asprintf_append(buffer, "[");
-					deref->image_index->accept(this);
-					ralloc_asprintf_append(buffer, "]");
-					if (Mask)
-					{
-						ralloc_asprintf_append(buffer, Mask);
-					}
-					ralloc_asprintf_append(buffer, " = ");
-					src->accept(this);
-				}
-			}
-			else
-			{
-				if (src == NULL)
-				{
-					ralloc_asprintf_append(buffer, "imageLoad( ");
-					deref->image->accept(this);
-					ralloc_asprintf_append(buffer, ", %s(", GLSLIntCastTypes[deref->image_index->type->vector_elements]);
-					deref->image_index->accept(this);
-					ralloc_asprintf_append(buffer, ")).%s", swizzle[dst_elements - 1]);
-				}
-				else
-				{
-					ralloc_asprintf_append(buffer, "imageStore( ");
-					deref->image->accept(this);
-					ralloc_asprintf_append(buffer, ", %s(", GLSLIntCastTypes[deref->image_index->type->vector_elements]);
-					deref->image_index->accept(this);
-					ralloc_asprintf_append(buffer, "), ");
-					// avoid 'scalar swizzle'
-					if (/*src->as_constant() && */src_elements == 1)
-					{
-						// Add cast if missing and avoid swizzle
-						if (deref->image->type->inner_type)
-						{
-							switch (deref->image->type->inner_type->base_type)
-							{
-							case GLSL_TYPE_INT:
-								ralloc_asprintf_append(buffer, "ivec4(");
-								break;
-							case GLSL_TYPE_UINT:
-								ralloc_asprintf_append(buffer, "uvec4(");
-								break;
-							case GLSL_TYPE_FLOAT:
-							case GLSL_TYPE_HALF:
-								ralloc_asprintf_append(buffer, "vec4(");
-								break;
-							default:
-								break;
-							}
-						}
-
-						src->accept(this);
-						ralloc_asprintf_append(buffer, "))");
-					}
-					else
-					{
-						src->accept(this);
-						ralloc_asprintf_append(buffer, ".%s)", expand[src_elements - 1]);
-					}
-				}
-			}
-		}
-		else if (deref->op == ir_image_dimensions) //-V547
-		{
-			check(!bIsStructured);
-			ralloc_asprintf_append(buffer, "imageSize( ");
-			deref->image->accept(this);
-			ralloc_asprintf_append(buffer, ")");
-		}
-		else
-		{
-			check(!bIsStructured);
-			check(!"Unknown image operation");
-		}
-	}
-
-	virtual void visit(ir_dereference_image *deref)
-	{
-		check(scope_depth > 0);
-
-		print_image_op(deref, NULL);
-
-	}
-
-	virtual void visit(ir_dereference_record *deref)
-	{
-		check(scope_depth > 0);
-
-		deref->record->accept(this);
-		ralloc_asprintf_append(buffer, ".%s", deref->field);
-
-		if (is_md_array(deref->type))
-		{
-			ralloc_asprintf_append(buffer, ".Inner");
-		}
-	}
-
-	virtual void visit(ir_assignment *assign)
-	{
-		if (scope_depth == 0)
-		{
-			global_instructions.push_tail(new(mem_ctx)global_ir(assign));
-			needs_semicolon = false;
-			return;
-		}
-
-		// constant variables with initializers are statically assigned
-		ir_variable *var = assign->lhs->variable_referenced();
-		if (var->has_initializer && var->read_only && (var->constant_initializer || var->constant_value))
-		{
-			//This will leave a blank line with a semi-colon
-			return;
-		}
-
-		if (assign->condition)
-		{
-			ralloc_asprintf_append(buffer, "if(");
-			assign->condition->accept(this);
-			ralloc_asprintf_append(buffer, ") { ");
-		}
-
-		char mask[6] = "";
-		unsigned j = 1;
-		if (assign->lhs->type->is_scalar() == false ||
-			assign->write_mask != 0x1)
-		{
-			for (unsigned i = 0; i < 4; i++)
-			{
-				if ((assign->write_mask & (1 << i)) != 0)
-				{
-					mask[j] = "xyzw"[i];
-					j++;
-				}
-			}
-		}
-		mask[j] = '\0';
-
-		mask[0] = (j == 1) ? '\0' : '.';
-
-		if (assign->lhs->as_dereference_image() != NULL)
-		{
-			print_image_op(assign->lhs->as_dereference_image(), assign->rhs, mask);
-		}
-		else
-		{
-			assign->lhs->accept(this);
-			ralloc_asprintf_append(buffer, "%s = ", mask);
-			assign->rhs->accept(this);
-		}
-
-		if (assign->condition)
-		{
-			ralloc_asprintf_append(buffer, "%s }", needs_semicolon ? ";" : "");
-		}
-	}
-
-	void print_constant(ir_constant *constant, int index)
-	{
-		if (constant->type->is_float())
-		{
-			if (constant->is_component_finite(index))
-			{
-				float value = constant->value.f[index];
-				// Original formatting code relied on %f style formatting
-				// %e is more accurate, and has been available since at least ES 2.0
-				// leaving original code in place, in case some drivers don't properly handle it
-#if 0
-				const char *format = (fabsf(fmodf(value, 1.0f)) < 1.e-8f) ? "%.1f" : "%.8f";
-				ralloc_asprintf_append(buffer, format, value);
-#else
-/*
-				f= 0;
-					%f =	0.000000
-					%e =	0.000000e+00
-					%.16g =	0
-					%g =	0
-				f=  256.0 / 255.0;
-					%f =	1.003922
-					%e =	1.003922e+00
-					%.16g =	1.003921627998352
-					%g =	1.00392
-				f=  2.3283064365386963e-10;
-					%f =	0.000000
-					%e =	2.328306e-10
-					%.16g =	2.328306436538696e-10
-					%g =	2.32831e-10
-				f=  1e-6;
-					%f =	0.000001
-					%e =	1.000000e-06
-					%.16g =	9.999999974752427e-07
-					%g =	1e-06
-				f=  -1;
-					%f =	-1.000000
-					%e =	-1.000000e+00
-					%.16g =	-1
-					%g =	-1
-				f=  -1000;
-					%f =	-1000.000000
-					%e =	-1.000000e+03
-					%.16g =	-1000
-					%g =	-1000
-				f=  -1048576;
-					%f =	-1048576.000000
-					%e =	-1.048576e+06
-					%.16g =	-1048576
-					%g =	-1.04858e+06
-				f=  UINT_MAX;
-					%f =	4294967296.000000
-					%e =	4.294967e+09
-					%.16g =	4294967296
-					%g =	4.29497e+09
-				f=	3.402823466e+38
-					%f=		340282346638528859811704183484516925440.000000
-					%e=		3.402823e+38
-					%.16g=	3.402823466385289e+38
-					%g=		3.40282e+38
-*/
-				// Not fast, but way more precise
-				{
-					// 128 to handle the case of 3.402823466e+38f (which we use!)
-					char f[128], g[128], g10[128];
-					float ReadValue;
-					// Always add decimal point as glsl is painful
-					FCStringAnsi::Sprintf(g, "%g", value);
-					Scanf(g, "%f", &ReadValue);
-					bool bHasDecimalPoint = FCStringAnsi::Strchr(g, '.') != nullptr;
-					if (bHasDecimalPoint && ReadValue == value)
-					{
-						ralloc_strcat(buffer, g);
-					}
-					else
-					{
-						FCStringAnsi::Sprintf(f, "%f", value);
-						Scanf(f, "%f", &ReadValue);
-						if (ReadValue == value)
-						{
-							ralloc_strcat(buffer, f);
-						}
-						else
-						{
-							bHasDecimalPoint = FCStringAnsi::Strchr(g, '.') != nullptr;
-							FCStringAnsi::Sprintf(g10, "%.10g", value);
-							Scanf(g10, "%f", &ReadValue);
-							if (ReadValue == value)
-							{
-								ralloc_strcat(buffer, g10);
-							}
-							else
-							{
-								ralloc_asprintf_append(buffer, "%.16e", value);
-							}
-						}
-					}
-				}
-#endif
-			}
-			else
-			{
-				switch (constant->value.u[index])
-				{
-				case 0x7f800000u:
-					ralloc_asprintf_append(buffer, "(1.0/0.0) /*Inf*/");
-					break;
-
-				case 0xff800000u:
-					ralloc_asprintf_append(buffer, "(-1.0/0.0) /*-Inf*/");
-					break;
-
-				default:
-					checkf(FGenericPlatformMath::IsNaN(constant->value.f[index]), TEXT("constant->value.u[index] = 0x%0x"), constant->value.u[index]);
-					if ((constant->value.u[index] >> 31u) == 1u)
-					{
-						ralloc_asprintf_append(buffer, "(-0.0/0.0) /*-Nan*/");
-					}
-					else
-					{
-						ralloc_asprintf_append(buffer, "(0.0/0.0) /*Nan*/");
-					}
-					break;
-				}
-			}
-		}
-		else if (constant->type->base_type == GLSL_TYPE_INT)
-		{
-			ralloc_asprintf_append(buffer, "%d", constant->value.i[index]);
-		}
-		else if (constant->type->base_type == GLSL_TYPE_UINT)
-		{
-			ralloc_asprintf_append(buffer, "%u%s",
-				constant->value.u[index],
-				should_print_uint_literals_as_ints ? "" : "u"
-				);
-		}
-		else if (constant->type->base_type == GLSL_TYPE_BOOL)
-		{
-			ralloc_asprintf_append(buffer, "%s", constant->value.b[index] ? "true" : "false");
-		}
-	}
-
-	virtual void visit(ir_constant *constant)
-	{
-		if (constant->type == glsl_type::float_type
-			|| constant->type == glsl_type::half_type
-			|| constant->type == glsl_type::bool_type
-			|| constant->type == glsl_type::int_type
-			|| constant->type == glsl_type::uint_type)
-		{
-			print_constant(constant, 0);
-		}
-		else if (constant->type->is_record())
-		{
-			print_type_full(constant->type);
-			ralloc_asprintf_append(buffer, "(");
-			ir_constant* value = (ir_constant*)constant->components.get_head();
-			if (value)
-			{
-				value->accept(this);
-			}
-			for (uint32 i = 1; i < constant->type->length; i++)
-			{
-				check(value);
-				value = (ir_constant*)value->next;
-				if (value)
-				{
-					ralloc_asprintf_append(buffer, ",");
-					value->accept(this);
-				}
-			}
-			ralloc_asprintf_append(buffer, ")");
-		}
-		else if (constant->type->is_array())
-		{
-			print_type_full(constant->type);
-			ralloc_asprintf_append(buffer, "(");
-			constant->get_array_element(0)->accept(this);
-			for (uint32 i = 1; i < constant->type->length; ++i)
-			{
-				ralloc_asprintf_append(buffer, ",");
-				constant->get_array_element(i)->accept(this);
-			}
-			ralloc_asprintf_append(buffer, ")");
-		}
-		else
-		{
-			print_type_full(constant->type);
-			ralloc_asprintf_append(buffer, "(");
-			print_constant(constant, 0);
-			int num_components = constant->type->components();
-			for (int i = 1; i < num_components; ++i)
-			{
-				ralloc_asprintf_append(buffer, ",");
-				print_constant(constant, i);
-			}
-			ralloc_asprintf_append(buffer, ")");
-		}
-	}
-
-	virtual void visit(ir_call *call)
-	{
-		if (scope_depth == 0)
-		{
-			global_instructions.push_tail(new(mem_ctx)global_ir(call));
-			needs_semicolon = false;
-			return;
-		}
-
-		if (call->return_deref)
-		{
-			call->return_deref->accept(this);
-			ralloc_asprintf_append(buffer, " = ");
-		}
-		ralloc_asprintf_append(buffer, "%s(", call->callee_name());
-		bool bPrintComma = false;
-		foreach_iter(exec_list_iterator, iter, *call)
-		{
-			ir_instruction *const inst = (ir_instruction *)iter.get();
-			if (bPrintComma)
-			{
-				ralloc_asprintf_append(buffer, ",");
-			}
-			inst->accept(this);
-			bPrintComma = true;
-		}
-		ralloc_asprintf_append(buffer, ")");
-	}
-
-	virtual void visit(ir_return *ret)
-	{
-		check(scope_depth > 0);
-
-		ralloc_asprintf_append(buffer, "return ");
-		ir_rvalue *const value = ret->get_value();
-		if (value)
-		{
-			value->accept(this);
-		}
-	}
-
-	virtual void visit(ir_discard *discard)
-	{
-		check(scope_depth > 0);
-
-		if (discard->condition)
-		{
-			ralloc_asprintf_append(buffer, "if (");
-			discard->condition->accept(this);
-			ralloc_asprintf_append(buffer, ") ");
-		}
-		ralloc_asprintf_append(buffer, "discard");
-		bUsesDiscard = true;
-	}
-
-	bool try_conditional_move(ir_if *expr)
-	{
-		ir_dereference_variable *dest_deref = NULL;
-		ir_rvalue *true_value = NULL;
-		ir_rvalue *false_value = NULL;
-		unsigned write_mask = 0;
-		const glsl_type *assign_type = NULL;
-		int num_inst;
-
-		num_inst = 0;
-		foreach_iter(exec_list_iterator, iter, expr->then_instructions)
-		{
-			if (num_inst > 0)
-			{
-				// multiple instructions? not a conditional move
-				return false;
-			}
-
-			ir_instruction *const inst = (ir_instruction *)iter.get();
-			ir_assignment *assignment = inst->as_assignment();
-			if (assignment && (assignment->rhs->ir_type == ir_type_dereference_variable || assignment->rhs->ir_type == ir_type_constant || assignment->rhs->ir_type == ir_type_swizzle))
-			{
-				dest_deref = assignment->lhs->as_dereference_variable();
-				true_value = assignment->rhs;
-				write_mask = assignment->write_mask;
-			}
-			num_inst++;
-		}
-
-		if (dest_deref == NULL || true_value == NULL)
-			return false;
-
-		num_inst = 0;
-		foreach_iter(exec_list_iterator, iter, expr->else_instructions)
-		{
-			if (num_inst > 0)
-			{
-				// multiple instructions? not a conditional move
-				return false;
-			}
-
-			ir_instruction *const inst = (ir_instruction *)iter.get();
-			ir_assignment *assignment = inst->as_assignment();
-			if (assignment && (assignment->rhs->ir_type == ir_type_dereference_variable || assignment->rhs->ir_type == ir_type_constant || assignment->rhs->ir_type == ir_type_swizzle))
-			{
-				ir_dereference_variable *tmp_deref = assignment->lhs->as_dereference_variable();
-				if (tmp_deref
-					&& tmp_deref->var == dest_deref->var
-					&& tmp_deref->type == dest_deref->type
-					&& assignment->write_mask == write_mask)
-				{
-					false_value = assignment->rhs;
-				}
-			}
-			num_inst++;
-		}
-
-		if (false_value == NULL)
-			return false;
-
-		char mask[6];
-		unsigned j = 1;
-		if (dest_deref->type->is_scalar() == false || write_mask != 0x1)
-		{
-			for (unsigned i = 0; i < 4; i++)
-			{
-				if ((write_mask & (1 << i)) != 0)
-				{
-					mask[j] = "xyzw"[i];
-					j++;
-				}
-			}
-		}
-		mask[j] = '\0';
-		mask[0] = (j == 1) ? '\0' : '.';
-
-		dest_deref->accept(this);
-		ralloc_asprintf_append(buffer, "%s = (", mask);
-		expr->condition->accept(this);
-		ralloc_asprintf_append(buffer, ")?(");
-		true_value->accept(this);
-		ralloc_asprintf_append(buffer, "):(");
-		false_value->accept(this);
-		ralloc_asprintf_append(buffer, ")");
-
-		return true;
-	}
-
-	virtual void visit(ir_if *expr)
-	{
-		check(scope_depth > 0);
-
-		if (try_conditional_move(expr) == false)
-		{
-			ralloc_asprintf_append(buffer, "if (");
-			expr->condition->accept(this);
-			ralloc_asprintf_append(buffer, ")\n");
-			indent();
-			ralloc_asprintf_append(buffer, "{\n");
-
-			indentation++;
-			foreach_iter(exec_list_iterator, iter, expr->then_instructions)
-			{
-				ir_instruction *const inst = (ir_instruction *)iter.get();
-				indent();
-				do_visit(inst);
-			}
-			indentation--;
-
-			indent();
-			ralloc_asprintf_append(buffer, "}\n");
-
-			if (!expr->else_instructions.is_empty())
-			{
-				indent();
-				ralloc_asprintf_append(buffer, "else\n");
-				indent();
-				ralloc_asprintf_append(buffer, "{\n");
-
-				indentation++;
-				foreach_iter(exec_list_iterator, iter, expr->else_instructions)
-				{
-					ir_instruction *const inst = (ir_instruction *)iter.get();
-					indent();
-					do_visit(inst);
-				}
-				indentation--;
-
-				indent();
-				ralloc_asprintf_append(buffer, "}\n");
-			}
-
-			needs_semicolon = false;
-		}
-	}
-
-	virtual void visit(ir_loop *loop)
-	{
-		check(scope_depth > 0);
-
-		if (loop->counter && loop->to)
-		{
-			// IR cmp operator is when to terminate loop; whereas GLSL for loop syntax
-			// is while to continue the loop. Invert the meaning of operator when outputting.
-			const char* termOp = NULL;
-			switch (loop->cmp)
-			{
-			case ir_binop_less: termOp = ">="; break;
-			case ir_binop_greater: termOp = "<="; break;
-			case ir_binop_lequal: termOp = ">"; break;
-			case ir_binop_gequal: termOp = "<"; break;
-			case ir_binop_equal: termOp = "!="; break;
-			case ir_binop_nequal: termOp = "=="; break;
-			default: check(false);
-			}
-			ralloc_asprintf_append(buffer, "for (;%s%s", unique_name(loop->counter), termOp);
-			loop->to->accept(this);
-			ralloc_asprintf_append(buffer, ";)\n");
-		}
-		else
-		{
-#if 1
-			ralloc_asprintf_append(buffer, "for (;;)\n");
-#else
-			ralloc_asprintf_append(buffer, "for ( int loop%d = 0; loop%d < 256; loop%d ++)\n", loop_count, loop_count, loop_count);
-			loop_count++;
-#endif
-		}
-		indent();
-		ralloc_asprintf_append(buffer, "{\n");
-
-		indentation++;
-		foreach_iter(exec_list_iterator, iter, loop->body_instructions)
-		{
-			ir_instruction *const inst = (ir_instruction *)iter.get();
-			indent();
-			do_visit(inst);
-		}
-		indentation--;
-
-		indent();
-		ralloc_asprintf_append(buffer, "}\n");
-
-		needs_semicolon = false;
-	}
-
-	virtual void visit(ir_loop_jump *jmp)
-	{
-		check(scope_depth > 0);
-
-		ralloc_asprintf_append(buffer, "%s",
-			jmp->is_break() ? "break" : "continue");
-	}
-
-	virtual void visit(ir_atomic *ir)
-	{
-		const char *sharedAtomicFunctions[] =
-		{
-			"atomicAdd",
-			"atomicAnd",
-			"atomicMin",
-			"atomicMax",
-			"atomicOr",
-			"atomicXor",
-			"atomicExchange",
-			"atomicCompSwap"
-		};
-		const char *imageAtomicFunctions[] =
-		{
-			"imageAtomicAdd",
-			"imageAtomicAnd",
-			"imageAtomicMin",
-			"imageAtomicMax",
-			"imageAtomicOr",
-			"imageAtomicXor",
-			"imageAtomicExchange",
-			"imageAtomicCompSwap"
-		};
-		check(scope_depth > 0);
-		ir_dereference_image* image = ir->memory_ref->as_dereference_image();
-
-		ir->lhs->accept(this);
-		if (!image || (image->image->type && image->image->type->shader_storage_buffer))
-		{
-			ralloc_asprintf_append(buffer, " = %s(",
-				sharedAtomicFunctions[ir->operation]);
-			ir->memory_ref->accept(this);
-			ralloc_asprintf_append(buffer, ", ");
-			ir->operands[0]->accept(this);
-			if (ir->operands[1])
-			{
-				ralloc_asprintf_append(buffer, ", ");
-				ir->operands[1]->accept(this);
-			}
-			ralloc_asprintf_append(buffer, ")");
-		}
-		else
-		{
-			ralloc_asprintf_append(buffer, " = %s(",
-				imageAtomicFunctions[ir->operation]);
-			image->image->accept(this);
-			ralloc_asprintf_append(buffer, ", %s(", GLSLIntCastTypes[image->image_index->type->vector_elements]);
-			image->image_index->accept(this);
-			ralloc_asprintf_append(buffer, "), ");
-			ir->operands[0]->accept(this);
-			if (ir->operands[1])
-			{
-				ralloc_asprintf_append(buffer, ", ");
-				ir->operands[1]->accept(this);
-			}
-			ralloc_asprintf_append(buffer, ")");
-
-			bUsesImageWriteAtomic = true;
-		}
-	}
-
-	void AddTypeToUsedStructs(const glsl_type* type);
-
-	/**
-	* Declare structs used to simulate multi-dimensional arrays.
-	*/
-	void declare_md_array_struct(const glsl_type* type, hash_table* ht)
-	{
-		check(type->is_array());
-
-		if (hash_table_find(ht, (void*)type) == NULL)
-		{
-			const glsl_type* subtype = type->fields.array;
-			if (subtype->base_type == GLSL_TYPE_ARRAY)
-			{
-				declare_md_array_struct(subtype, ht);
-
-				ralloc_asprintf_append(buffer, "struct ");
-				print_md_array_type(type);
-				ralloc_asprintf_append(buffer, "\n{\n\t");
-				print_md_array_type(subtype);
-				ralloc_asprintf_append(buffer, " Inner[%u];\n};\n\n", type->length);
-			}
-			else
-			{
-				ralloc_asprintf_append(buffer, "struct ");
-				print_md_array_type(type);
-				ralloc_asprintf_append(buffer, "\n{\n\t");
-				print_type_pre(type);
-				ralloc_asprintf_append(buffer, " Inner");
-				print_type_post(type);
-				ralloc_asprintf_append(buffer, ";\n};\n\n");
-			}
-			hash_table_insert(ht, (void*)type, (void*)type);
-		}
-	}
-
-	/**
-	* Declare structs used by the code that has been generated.
-	*/
-	void declare_structs(_mesa_glsl_parse_state* state, bool bCanHaveUBs)
-	{
-		// If any variable in a uniform block is in use, the entire uniform block
-		// must be present including structs that are not actually accessed.
-		for (unsigned i = 0; i < state->num_uniform_blocks; i++)
-		{
-			const glsl_uniform_block* block = state->uniform_blocks[i];
-			if (hash_table_find(used_uniform_blocks, block->name))
-			{
-				for (unsigned var_index = 0; var_index < block->num_vars; ++var_index)
-				{
-					const glsl_type* type = block->vars[var_index]->type;
-					if (type->base_type == GLSL_TYPE_STRUCT &&
-						hash_table_find(used_structures, type) == NULL)
-					{
-						hash_table_insert(used_structures, (void*)type, type);
-					}
-				}
-			}
-		}
-
-		// If otherwise unused structure is a member of another, used structure, the unused structure is also, in fact, used
-		{
-			int added_structure_types;
-			do
-			{
-				added_structure_types = 0;
-				for (unsigned i = 0; i < state->num_user_structures; i++)
-				{
-					const glsl_type *const s = state->user_structures[i];
-
-					if (hash_table_find(used_structures, s) == NULL)
-					{
-						continue;
-					}
-
-					for (unsigned j = 0; j < s->length; j++)
-					{
-						const glsl_type* type = s->fields.structure[j].type;
-
-						if (type->base_type == GLSL_TYPE_STRUCT)
-						{
-							if (hash_table_find(used_structures, type) == NULL)
-							{
-								hash_table_insert(used_structures, (void*)type, type);
-								++added_structure_types;
-							}
-						}
-						else if (type->base_type == GLSL_TYPE_ARRAY && type->fields.array->base_type == GLSL_TYPE_STRUCT)
-						{
-							if (hash_table_find(used_structures, type->fields.array) == NULL)
-							{
-								hash_table_insert(used_structures, (void*)type->fields.array, type->fields.array);
-							}
-						}
-						else if ((type->base_type == GLSL_TYPE_INPUTPATCH || type->base_type == GLSL_TYPE_OUTPUTPATCH) && type->inner_type->base_type == GLSL_TYPE_STRUCT)
-						{
-							if (hash_table_find(used_structures, type->inner_type) == NULL)
-							{
-								hash_table_insert(used_structures, (void*)type->inner_type, type->inner_type);
-							}
-						}
-					}
-				}
-			} while (added_structure_types > 0);
-		}
-
-		// Generate structures that allow support for multi-dimensional arrays.
-		{
-			hash_table* ht = hash_table_ctor(32, hash_table_pointer_hash, hash_table_pointer_compare);
-			foreach_iter(exec_list_iterator, iter, used_md_arrays)
-			{
-				md_array_entry* entry = (md_array_entry*)iter.get();
-				declare_md_array_struct(entry->type, ht);
-			}
-			hash_table_dtor(ht);
-		}
-
-		for (unsigned i = 0; i < state->num_user_structures; i++)
-		{
-			const glsl_type *const s = state->user_structures[i];
-
-			if (hash_table_find(used_structures, s) == NULL)
-			{
-				continue;
-			}
-
-			ralloc_asprintf_append(buffer, "struct %s\n{\n", s->name);
-
-			if (s->length == 0)
-			{
-				if (bEmitPrecision)
-				{
-					ralloc_asprintf_append(buffer, "\thighp float glsl_doesnt_like_empty_structs;\n");
-				}
-				else
-				{
-					ralloc_asprintf_append(buffer, "\tfloat glsl_doesnt_like_empty_structs;\n");
-				}
-			}
-			else
-			{
-				for (unsigned j = 0; j < s->length; j++) 
-				{
-					const bool bSupportsPrecisionQualifier = GlslTypeSupportsPrecisionQualifier(*(s->fields.structure[j].type));
-					ralloc_asprintf_append(buffer, "\t%s ", (state->language_version == 310 && bEmitPrecision && bSupportsPrecisionQualifier) ? "highp" : "");
-					print_type_pre(s->fields.structure[j].type);
-					ralloc_asprintf_append(buffer, " %s", s->fields.structure[j].name);
-					print_type_post(s->fields.structure[j].type);
-					ralloc_asprintf_append(buffer, ";\n");
-				}
-			}
-			ralloc_asprintf_append(buffer, "};\n\n");
-		}
-
-		// Non Global UBs; if bCanHaveUBs then we can't assume they are all packed
-		unsigned num_used_blocks = 0;
-		for (unsigned i = 0; i < state->num_uniform_blocks; i++)
-		{
-			const glsl_uniform_block* block = state->uniform_blocks[i];
-			if (hash_table_find(used_uniform_blocks, block->name))
-			{
-				const char* block_name = block->name;
-
-				check(block->num_vars > 0);
-				const char* var_name = block->vars[0]->name;
-
-/*
-				block_name = ralloc_asprintf(mem_ctx, "%sb%u",
-					glsl_variable_tag_from_parser_target(state->target),
-					num_used_blocks
-					);
-*/
-
-				auto Type = EVulkanBindingType::UniformBuffer;
-				if (bCanHaveUBs && block->num_vars == 1 && strlen(var_name) == 4 && var_name[0] == glsl_variable_tag_from_parser_target(state->target)[0] && var_name[1] == 'u' && var_name[2] == '_')
-				{
-					// Find in the regular globals
-					auto Found = state->GlobalPackedArraysMap.find(var_name[3]);
-					if (Found != state->GlobalPackedArraysMap.end())
-					{
-						Type = EVulkanBindingType::PackedUniformBuffer;
-					}
-					else
-					{
-						// Find in the emulated UBs
-						for (auto Pair : state->CBPackedArraysMap)
-						{
-							auto InnerFound = Pair.second.find(var_name[3]);
-							if (InnerFound != Pair.second.end())
-							{
-								Type = EVulkanBindingType::PackedUniformBuffer;
-								break;
-							}
-						}
-					}
-				}
-
-				int32 Binding = BindingTable.RegisterBinding(block_name, var_name, Type);
-				ralloc_asprintf_append(
-					buffer,
-					"layout(set=%d, binding=BINDING_%d, std140) uniform %s\n{\n",
-					GetDescriptorSetForStage(ParseState->target),
-					Binding,
-					block_name);
-
-				bool optimized_structure_out = false;
-				if (!optimized_structure_out)
-				{
-					for (unsigned var_index = 0; var_index < block->num_vars; ++var_index)
-					{
-						ir_variable* var = block->vars[var_index];
-
-						//EHart - name-mangle variables to prevent colliding names
-						//#todo-rco: Check if this is still is needed when creating PSOs
-						//ralloc_asprintf_append(buffer, "#define %s %s%s\n", var->name, var->name, block_name);
-						const bool bSupportsPrecisionQualifier = GlslTypeSupportsPrecisionQualifier(*(var->type));
-						ralloc_asprintf_append(buffer, "\t%s", (state->language_version == 310 && bEmitPrecision && bSupportsPrecisionQualifier) ? "highp " : "");
-						print_type_pre(var->type);
-						ralloc_asprintf_append(buffer, " %s", var->name);
-						print_type_post(var->type);
-						ralloc_asprintf_append(buffer, ";\n");
-					}
-					ralloc_asprintf_append(buffer, "};\n\n");
-				}
-
-				num_used_blocks++;
-			}
-		}
-	}
-
-	void PrintPackedSamplers(_mesa_glsl_parse_state::TUniformList& Samplers, TStringToSetMap& TextureToSamplerMap)
-	{
-		bool bPrintHeader = true;
-		bool bNeedsComma = false;
-		for (_mesa_glsl_parse_state::TUniformList::iterator Iter = Samplers.begin(); Iter != Samplers.end(); ++Iter)
-		{
-			glsl_packed_uniform& Sampler = *Iter;
-			std::string SamplerStates("");
-			TStringToSetMap::iterator IterFound = TextureToSamplerMap.find(Sampler.Name);
-			if (IterFound != TextureToSamplerMap.end())
-			{
-				TStringSet& ListSamplerStates = IterFound->second;
-				check(!ListSamplerStates.empty());
-				for (TStringSet::iterator IterSS = ListSamplerStates.begin(); IterSS != ListSamplerStates.end(); ++IterSS)
-				{
-					if (IterSS == ListSamplerStates.begin())
-					{
-						SamplerStates += "[";
-					}
-					else
-					{
-						SamplerStates += ",";
-					}
-					SamplerStates += *IterSS;
-				}
-
-				SamplerStates += "]";
-			}
-
-			ralloc_asprintf_append(
-				buffer,
-				"%s%s(%u:%u%s)",
-				bNeedsComma ? "," : "",
-				Sampler.Name.c_str(),
-				Sampler.offset,
-				Sampler.num_components,
-				SamplerStates.c_str()
-				);
-
-			bNeedsComma = true;
-		}
-		/*
-		for (TStringToSetMap::iterator Iter = state->TextureToSamplerMap.begin(); Iter != state->TextureToSamplerMap.end(); ++Iter)
-		{
-		const std::string& Texture = Iter->first;
-		TStringSet& Samplers = Iter->second;
-		if (!Samplers.empty())
-		{
-		if (bFirstTexture)
-		{
-		bFirstTexture = false;
-		}
-		else
-		{
-		ralloc_asprintf_append(buffer, ",");
-		}
-
-		ralloc_asprintf_append(buffer, "%s(", Texture.c_str());
-		bool bFirstSampler = true;
-		for (TStringSet::iterator IterSamplers = Samplers.begin(); IterSamplers != Samplers.end(); ++IterSamplers)
-		{
-		if (bFirstSampler)
-		{
-		bFirstSampler = false;
-		}
-		else
-		{
-		ralloc_asprintf_append(buffer, ",");
-		}
-
-		ralloc_asprintf_append(buffer, "%s", IterSamplers->c_str());
-		}
-		ralloc_asprintf_append(buffer, ")");
-		}
-		}
-		*/
-	}
-
-	bool PrintPackedUniforms(bool bPrintArrayType, char ArrayType, _mesa_glsl_parse_state::TUniformList& Uniforms, bool bFlattenUniformBuffers, bool NeedsComma)
-	{
-		bool bPrintHeader = true;
-		for (glsl_packed_uniform& Uniform : Uniforms)
-		{
-			if (!bFlattenUniformBuffers || Uniform.CB_PackedSampler.empty())
-			{
-				if (bPrintArrayType && bPrintHeader)
-				{
-					ralloc_asprintf_append(buffer, "%s%c[",
-						NeedsComma ? "," : "",
-						ArrayType);
-					bPrintHeader = false;
-					NeedsComma = false;
-				}
-				ralloc_asprintf_append(
-					buffer,
-					"%s%s(%u:%u)",
-					NeedsComma ? "," : "",
-					Uniform.Name.c_str(),
-					Uniform.offset,
-					Uniform.num_components
-					);
-				NeedsComma = true;
-			}
-		}
-
-		if (bPrintArrayType && !bPrintHeader)
-		{
-			ralloc_asprintf_append(buffer, "]");
-		}
-
-		return NeedsComma;
-	}
-
-	void PrintPackedGlobals(_mesa_glsl_parse_state* State)
-	{
-		//	@PackedGlobals: Global0(DestArrayType, DestOffset, SizeInFloats), Global1(DestArrayType, DestOffset, SizeInFloats), ...
-		bool bNeedsHeader = true;
-		bool bNeedsComma = false;
-		for (auto& Pair : State->GlobalPackedArraysMap)
-		{
-			char ArrayType = Pair.first;
-			if (ArrayType != EArrayType_Image && ArrayType != EArrayType_Sampler)
-			{
-				_mesa_glsl_parse_state::TUniformList& Uniforms = Pair.second;
-				check(!Uniforms.empty());
-
-				for (auto Iter = Uniforms.begin(); Iter != Uniforms.end(); ++Iter)
-				{
-					glsl_packed_uniform& Uniform = *Iter;
-					if (!State->bFlattenUniformBuffers || Uniform.CB_PackedSampler.empty())
-					{
-						if (bNeedsHeader)
-						{
-							ralloc_asprintf_append(buffer, "// @PackedGlobals: ");
-							bNeedsHeader = false;
-						}
-
-						ralloc_asprintf_append(
-							buffer,
-							"%s%s(%c:%u,%u)",
-							bNeedsComma ? "," : "",
-							Uniform.Name.c_str(),
-							ArrayType,
-							Uniform.offset,
-							Uniform.num_components
-							);
-						bNeedsComma = true;
-					}
-				}
-			}
-		}
-
-		if (!bNeedsHeader)
-		{
-			ralloc_asprintf_append(buffer, "\n");
-		}
-	}
-
-	void PrintPackedUniformBuffers(_mesa_glsl_parse_state* State, bool bGroupFlattenedUBs)
-	{
-		// @PackedUB: UniformBuffer0(SourceIndex0): Member0(SourceOffset,SizeInFloats),Member1(SourceOffset,SizeInFloats), ...
-		// @PackedUB: UniformBuffer1(SourceIndex1): Member0(SourceOffset,SizeInFloats),Member1(SourceOffset,SizeInFloats), ...
-		// ...
-
-		// First find all used CBs (since we lost that info during flattening)
-		TStringSet UsedCBs;
-		for (auto IterCB = State->CBPackedArraysMap.begin(); IterCB != State->CBPackedArraysMap.end(); ++IterCB)
-		{
-			for (auto Iter = IterCB->second.begin(); Iter != IterCB->second.end(); ++Iter)
-			{
-				_mesa_glsl_parse_state::TUniformList& Uniforms = Iter->second;
-				for (auto IterU = Uniforms.begin(); IterU != Uniforms.end(); ++IterU)
-				{
-					if (!IterU->CB_PackedSampler.empty())
-					{
-						check(IterCB->first == IterU->CB_PackedSampler);
-						UsedCBs.insert(IterU->CB_PackedSampler);
-					}
-				}
-			}
-		}
-
-		check(UsedCBs.size() == State->CBPackedArraysMap.size());
-
-		// Now get the CB index based off source declaration order, and print an info line for each, while creating the mem copy list
-		unsigned CBIndex = 0;
-		TCBDMARangeMap CBRanges;
-		for (unsigned i = 0; i < State->num_uniform_blocks; i++)
-		{
-			const glsl_uniform_block* block = State->uniform_blocks[i];
-			if (UsedCBs.find(block->name) != UsedCBs.end())
-			{
-				bool bNeedsHeader = true;
-
-				// Now the members for this CB
-				bool bNeedsComma = false;
-				auto IterPackedArrays = State->CBPackedArraysMap.find(block->name);
-				check(IterPackedArrays != State->CBPackedArraysMap.end());
-				for (auto Iter = IterPackedArrays->second.begin(); Iter != IterPackedArrays->second.end(); ++Iter)
-				{
-					char ArrayType = Iter->first;
-					check(ArrayType != EArrayType_Image && ArrayType != EArrayType_Sampler);
-
-					_mesa_glsl_parse_state::TUniformList& Uniforms = Iter->second;
-					for (auto IterU = Uniforms.begin(); IterU != Uniforms.end(); ++IterU)
-					{
-						glsl_packed_uniform& Uniform = *IterU;
-						if (Uniform.CB_PackedSampler == block->name)
-						{
-							if (bNeedsHeader)
-							{
-								ralloc_asprintf_append(buffer, "// @PackedUB: %s(%d): ",
-									block->name,
-									CBIndex);
-								bNeedsHeader = false;
-							}
-
-							ralloc_asprintf_append(buffer, "%s%s(%u,%u)",
-								bNeedsComma ? "," : "",
-								Uniform.Name.c_str(),
-								Uniform.OffsetIntoCBufferInFloats,
-								Uniform.SizeInFloats);
-
-							bNeedsComma = true;
-							unsigned SourceOffset = Uniform.OffsetIntoCBufferInFloats;
-							unsigned DestOffset = Uniform.offset;
-							unsigned Size = Uniform.SizeInFloats;
-							unsigned DestCBIndex = bGroupFlattenedUBs ? std::distance(UsedCBs.begin(), UsedCBs.find(block->name)) : 0;
-							unsigned DestCBPrecision = ArrayType;
-							InsertRange(CBRanges, CBIndex, SourceOffset, Size, DestCBIndex, DestCBPrecision, DestOffset);
-						}
-					}
-				}
-
-				if (!bNeedsHeader)
-				{
-					ralloc_asprintf_append(buffer, "\n");
-				}
-
-				CBIndex++;
-			}
-		}
-
-		//DumpSortedRanges(SortRanges(CBRanges));
-
-		// @PackedUBCopies: SourceArray:SourceOffset-DestArray:DestOffset,SizeInFloats;SourceArray:SourceOffset-DestArray:DestOffset,SizeInFloats,...
-		bool bFirst = true;
-		for (auto& Pair : CBRanges)
-		{
-			TDMARangeList& List = Pair.second;
-			for (auto IterList = List.begin(); IterList != List.end(); ++IterList)
-			{
-				if (bFirst)
-				{
-					ralloc_asprintf_append(buffer, bGroupFlattenedUBs ? "// @PackedUBCopies: " : "// @PackedUBGlobalCopies: ");
-					bFirst = false;
-				}
-				else
-				{
-					ralloc_asprintf_append(buffer, ",");
-				}
-
-				if (bGroupFlattenedUBs)
-				{
-					ralloc_asprintf_append(buffer, "%u:%u-%u:%c:%u:%u", IterList->SourceCB, IterList->SourceOffset, IterList->DestCBIndex, IterList->DestCBPrecision, IterList->DestOffset, IterList->Size);
-				}
-				else
-				{
-					check(IterList->DestCBIndex == 0);
-					ralloc_asprintf_append(buffer, "%u:%u-%c:%u:%u", IterList->SourceCB, IterList->SourceOffset, IterList->DestCBPrecision, IterList->DestOffset, IterList->Size);
-				}
-			}
-		}
-
-		if (!bFirst)
-		{
-			ralloc_asprintf_append(buffer, "\n");
-		}
-	}
-
-	void PrintPackedUniforms(_mesa_glsl_parse_state* State, bool bGroupFlattenedUBs)
-	{
-		PrintPackedGlobals(State);
-
-		if (State->bFlattenUniformBuffers && !State->CBuffersOriginal.empty())
-		{
-			PrintPackedUniformBuffers(State, bGroupFlattenedUBs);
-		}
-	}
-
-	/**
-	* Print a list of external variables.
-	*/
-	void print_extern_vars(_mesa_glsl_parse_state* State, exec_list* extern_vars)
-	{
-		const char *type_str[GLSL_TYPE_MAX] = { "u", "i", "f", "f", "b", "t", "?", "?", "?", "?", "s", "os", "im", "ip", "op" };
-		const char *col_str[] = { "", "", "2x", "3x", "4x" };
-		const char *row_str[] = { "", "1", "2", "3", "4" };
-
-		check(sizeof(type_str) / sizeof(char*) == GLSL_TYPE_MAX);
-
-		bool need_comma = false;
-		foreach_iter(exec_list_iterator, iter, *extern_vars)
-		{
-			ir_variable* var = ((extern_var*)iter.get())->var;
-			const glsl_type* type = var->type;
-			if (!FCStringAnsi::Strcmp(var->name, "gl_in"))
-			{
-				// Ignore it, as we can't properly frame this information in current format, and it's not used anyway for geometry shaders
-				continue;
-			}
-			if (!strncmp(var->name, "in_", 3) || !strncmp(var->name, "out_", 4))
-			{
-				if (type->is_record())
-				{
-					// This is the specific case for GLSL >= 150, as we generate a struct with a member for each interpolator (which we still want to count)
-					if (type->length != 1)
-					{
-						_mesa_glsl_warning(State, "Found a complex structure as in/out, counting is not implemented yet...\n");
-						continue;
-					}
-
-					type = type->fields.structure->type;
-				}
-			}
-			check(type);
-			bool is_array = type->is_array();
-			int array_size = is_array ? type->length : 0;
-			if (is_array)
-			{
-				type = type->fields.array;
-			}
-			ralloc_asprintf_append(buffer, "%s%s%s%s",
-				need_comma ? "," : "",
-				type->base_type == GLSL_TYPE_STRUCT ? type->name : type_str[type->base_type],
-				col_str[type->matrix_columns],
-				row_str[type->vector_elements]);
-			if (is_array)
-			{
-				ralloc_asprintf_append(buffer, "[%u]", array_size);
-			}
-			ralloc_asprintf_append(buffer, ";%d:%s", var->location, var->name);
-			need_comma = true;
-		}
-	}
-
-	/**
-	* Print the input/output signature for this shader.
-	*/
-	void print_signature(_mesa_glsl_parse_state *state, bool bGroupFlattenedUBs)
-	{
-		if (!input_variables.is_empty())
-		{
-			ralloc_asprintf_append(buffer, "// @Inputs: ");
-			print_extern_vars(state, &input_variables);
-			ralloc_asprintf_append(buffer, "\n");
-		}
-		if (!output_variables.is_empty())
-		{
-			ralloc_asprintf_append(buffer, "// @Outputs: ");
-			print_extern_vars(state, &output_variables);
-			ralloc_asprintf_append(buffer, "\n");
-		}
-		if (state->num_uniform_blocks > 0 && !state->bFlattenUniformBuffers)
-		{
-			bool first = true;
-			int Index = 0;
-			for (unsigned i = 0; i < state->num_uniform_blocks; i++)
-			{
-				const glsl_uniform_block* block = state->uniform_blocks[i];
-				if (hash_table_find(used_uniform_blocks, block->name))
-				{
-					ralloc_asprintf_append(buffer, "%s%s(%d)",
-						first ? "// @UniformBlocks: " : ",",
-						block->name, Index);
-					first = false;
-					++Index;
-				}
-			}
-			if (!first)
-			{
-				ralloc_asprintf_append(buffer, "\n");
-			}
-		}
-
-		if (state->has_packed_uniforms)
-		{
-			PrintPackedUniforms(state, bGroupFlattenedUBs);
-
-			if (!state->GlobalPackedArraysMap[EArrayType_Sampler].empty())
-			{
-				ralloc_asprintf_append(buffer, "// @Samplers: ");
-				PrintPackedSamplers(
-					state->GlobalPackedArraysMap[EArrayType_Sampler],
-					state->TextureToSamplerMap
-					);
-				ralloc_asprintf_append(buffer, "\n");
-			}
-
-			if (!state->GlobalPackedArraysMap[EArrayType_Image].empty())
-			{
-				ralloc_asprintf_append(buffer, "// @UAVs: ");
-				PrintPackedUniforms(
-					false,
-					EArrayType_Image,
-					state->GlobalPackedArraysMap[EArrayType_Image],
-					false,
-					false
-					);
-				ralloc_asprintf_append(buffer, "\n");
-			}
-		}
-		else
-		{
-			if (!uniform_variables.is_empty())
-			{
-				ralloc_asprintf_append(buffer, "// @Uniforms: ");
-				print_extern_vars(state, &uniform_variables);
-				ralloc_asprintf_append(buffer, "\n");
-			}
-			if (!sampler_variables.is_empty())
-			{
-				ralloc_asprintf_append(buffer, "// @Samplers: ");
-				print_extern_vars(state, &sampler_variables);
-				ralloc_asprintf_append(buffer, "\n");
-			}
-			if (!image_variables.is_empty())
-			{
-				ralloc_asprintf_append(buffer, "// @UAVs: ");
-				print_extern_vars(state, &image_variables);
-				ralloc_asprintf_append(buffer, "\n");
-			}
-		}
-
-		if (!SamplerStateNames.empty())
-		{
-			ralloc_asprintf_append(buffer, "// @SamplerStates: ");
-			for (uint32 Index = 0; Index < SamplerStateNames.size(); ++Index)
-			{
-				ralloc_asprintf_append(buffer, "%s%d:%s", Index > 0 ? "," : "", Index, SamplerStateNames[Index].c_str());
-			}
-			ralloc_asprintf_append(buffer, "\n");
-		}
-
-		if (!ExternalSamplersList.empty())
-		{
-			ralloc_asprintf_append(buffer, "// @ExternalTextures: ");
-			for (uint32 Index = 0; Index < ExternalSamplersList.size(); ++Index)
-			{
-				ralloc_asprintf_append(buffer, "%s%s", Index == 0 ? "" : ",", ExternalSamplersList[Index].c_str());
-			}
-			ralloc_asprintf_append(buffer, "\n");
-		}
-	}
-
-	/**
-	* Print the layout directives for this shader.
-	*/
-	void print_layout(_mesa_glsl_parse_state *state)
-	{
-		if (early_depth_stencil && this->bUsesDiscard == false)
-		{
-			ralloc_asprintf_append(buffer, "layout(early_fragment_tests) in;\n");
-		}
-		if (state->target == compute_shader)
-		{
-			ralloc_asprintf_append(buffer, "layout( local_size_x = %d, "
-				"local_size_y = %d, local_size_z = %d ) in;\n", wg_size_x,
-				wg_size_y, wg_size_z);
-		}
-
-		if (state->target == tessellation_control_shader)
-		{
-			ralloc_asprintf_append(buffer, "layout(vertices = %d) out;\n", tessellation.outputcontrolpoints);
-		}
-
-		if (state->target == tessellation_evaluation_shader)
-		{
-
-			std::stringstream str;
-
-			switch (tessellation.outputtopology)
-			{
-				// culling is inverted, see TranslateCullMode in the OpenGL and D3D11 RHI
-			case GLSL_OUTPUTTOPOLOGY_POINT:
-				str << "point_mode";
-				break;
-			case GLSL_OUTPUTTOPOLOGY_LINE:
-				str << "iso_lines";
-				break;
-
-			default:
-			case GLSL_OUTPUTTOPOLOGY_NONE:
-			case GLSL_OUTPUTTOPOLOGY_TRIANGLE_CW:
-				str << "triangles, ccw";
-				break;
-			case GLSL_OUTPUTTOPOLOGY_TRIANGLE_CCW:
-				str << "triangles, cw";
-				break;
-			}
-
-			switch (tessellation.partitioning)
-			{
-			default:
-			case GLSL_PARTITIONING_NONE:
-			case GLSL_PARTITIONING_INTEGER:
-				str << ", equal_spacing";
-				break;
-			case GLSL_PARTITIONING_FRACTIONAL_EVEN:
-				str << ", fractional_even_spacing";
-				break;
-			case GLSL_PARTITIONING_FRACTIONAL_ODD:
-				str << ", fractional_odd_spacing";
-				break;
-				// that assumes that the hull/control shader clamps the tessellation factors to be power of two
-			case GLSL_PARTITIONING_POW2:
-				str << ", equal_spacing";
-				break;
-			}
-			ralloc_asprintf_append(buffer, "layout(%s) in;\n", str.str().c_str());
-		}
-#if 0
-		if (state->target == tessellation_evaluation_shader || state->target == tessellation_control_shader)
-		{
-			ralloc_asprintf_append(buffer, "/* DEBUG DUMP\n");
-
-			ralloc_asprintf_append(buffer, "tessellation.domain =  %s \n", DomainStrings[tessellation.domain]);
-			ralloc_asprintf_append(buffer, "tessellation.outputtopology =  %s \n", OutputTopologyStrings[tessellation.outputtopology]);
-			ralloc_asprintf_append(buffer, "tessellation.partitioning =  %s \n", PartitioningStrings[tessellation.partitioning]);
-			ralloc_asprintf_append(buffer, "tessellation.maxtessfactor =  %f \n", tessellation.maxtessfactor);
-			ralloc_asprintf_append(buffer, "tessellation.outputcontrolpoints =  %d \n", tessellation.outputcontrolpoints);
-			ralloc_asprintf_append(buffer, "tessellation.patchconstantfunc =  %s \n", tessellation.patchconstantfunc);
-			ralloc_asprintf_append(buffer, " */\n");
-		}
-#endif		
-	}
-
-	void print_extensions(_mesa_glsl_parse_state* state, bool bUsesES31Extensions, bool bShouldEmitOESExtensions, bool bShouldEmitMultiview)
-	{
-		if (bUsesGLTextureLODExtension)
-		{
-			//ralloc_asprintf_append(buffer, "#ifndef DONTEMITEXTENSIONSHADERTEXTURELODENABLE\n");
-			//ralloc_asprintf_append(buffer, "#extension GL_EXT_shader_texture_lod : enable\n");
-			//ralloc_asprintf_append(buffer, "#endif\n");
-		}
-
-		if (state->bSeparateShaderObjects && !state->bGenerateES &&
-			((state->target == tessellation_control_shader) || (state->target == tessellation_evaluation_shader)))
-		{
-			ralloc_asprintf_append(buffer, "#extension GL_ARB_tessellation_shader : enable\n");
-		}
-
-		if (bUsesDXDY && bIsES)
-		{
-			ralloc_asprintf_append(buffer, "#extension GL_OES_standard_derivatives : enable\n");
-		}
-
-		if (bUsesImageWriteAtomic && bShouldEmitOESExtensions)
-		{
-			ralloc_asprintf_append(buffer, "#extension GL_OES_shader_image_atomic : enable\n");
-		}
-
-		if (bUsesES31Extensions)
-		{
-			ralloc_asprintf_append(buffer, "#extension GL_EXT_gpu_shader5 : enable\n");
-			ralloc_asprintf_append(buffer, "#extension GL_EXT_texture_buffer : enable\n");
-			ralloc_asprintf_append(buffer, "#extension GL_EXT_texture_cube_map_array : enable\n");
-			ralloc_asprintf_append(buffer, "#extension GL_EXT_shader_io_blocks : enable\n");
-
-			if (ParseState->target == geometry_shader)
-			{
-				ralloc_asprintf_append(buffer, "#extension GL_EXT_geometry_shader : enable\n");
-			}
-			else if (ParseState->target == tessellation_control_shader || ParseState->target == tessellation_evaluation_shader)
-			{
-				ralloc_asprintf_append(buffer, "#extension GL_EXT_tessellation_shader : enable\n");
-			}
-			else if (ParseState->target == compute_shader)
-			{
-				ralloc_asprintf_append(buffer, "#extension GL_OES_shader_image_atomic : enable\n");
-			}
-		}
-
-		if (bShouldEmitMultiview && (state->target == vertex_shader || state->target == fragment_shader))
-		{
-			ralloc_asprintf_append(buffer, "#extension GL_EXT_multiview : enable\n");
-		}
-	}
-
-public:
-	FSamplerMapping SamplerMapping;
-
-	/** Constructor. */
-	FGenerateVulkanVisitor(EHlslCompileTarget InTarget,
-							FVulkanBindingTable& InBindingTable,
-							_mesa_glsl_parse_state* InState,
-							bool bInGenerateLayoutLocations,
-							bool bInDefaultPrecisionIsHalf)
-		: early_depth_stencil(false)
-		, wg_size_x(0)
-		, wg_size_y(0)
-		, wg_size_z(0)
-		, Target(InTarget)
-		, ParseState(InState)
-		, bGenerateLayoutLocations(bInGenerateLayoutLocations)
-		, bDefaultPrecisionIsHalf(bInDefaultPrecisionIsHalf)
-		, BindingTable(InBindingTable)
-		, buffer(0)
-		, indentation(0)
-		, scope_depth(0)
-		, temp_id(0)
-		, global_id(0)
-		, needs_semicolon(false)
-		, should_print_uint_literals_as_ints(false)
-		, loop_count(0)
-		, bUsesGLTextureLODExtension(false)
-		, bUsesDXDY(false)
-		, bUsesDiscard(false)
-		, bUsesImageWriteAtomic(false)
-	{
-		printable_names = hash_table_ctor(32, hash_table_pointer_hash, hash_table_pointer_compare);
-		used_structures = hash_table_ctor(32, hash_table_pointer_hash, hash_table_pointer_compare);
-		used_uniform_blocks = hash_table_ctor(32, hash_table_string_hash, hash_table_string_compare);
-
-		bEmitPrecision = (Target == HCT_FeatureLevelES3_1 || Target == HCT_FeatureLevelES3_1Ext);
-		bIsES = false;//(Target == HCT_FeatureLevelES3_1 || Target == HCT_FeatureLevelES3_1Ext);
-		bIsES31 = (Target == HCT_FeatureLevelES3_1 || Target == HCT_FeatureLevelES3_1Ext);
-	}
-
-	/** Destructor. */
-	virtual ~FGenerateVulkanVisitor()
-	{
-		hash_table_dtor(printable_names);
-		hash_table_dtor(used_structures);
-		hash_table_dtor(used_uniform_blocks);
-	}
-
-	void FindAtomicVariables(exec_list* ir)
-	{
-		::FindAtomicVariables(ir, AtomicVariables);
-	}
-
-	int32 AddUniqueSamplerState(const std::string& Name)
-	{
-		if (Name == "")
-		{
-			return - 1;
-		}
-
-		for (uint32 Index = 0; Index < SamplerStateNames.size(); ++Index)
-		{
-			if (SamplerStateNames[Index] == Name)
-			{
-				return (int32)Index;
-			}
-		}
-
-		SamplerStateNames.push_back(Name);
-		return SamplerStateNames.size() - 1;
-	};
-
-	/**
-	* Executes the visitor on the provided ir.
-	* @returns the GLSL source code generated.
-	*/
-	const char* run(exec_list* ir, _mesa_glsl_parse_state* state, bool bGroupFlattenedUBs, bool bCanHaveUBs, uint32 SubpassFetchMask)
-	{
-		mem_ctx = ralloc_context(NULL);
-
-		char* code_buffer = ralloc_asprintf(mem_ctx, "");
-		buffer = &code_buffer;
-
-		if (bEmitPrecision && ParseState->target == fragment_shader)
-		{
-			// TODO: Improve this...
-
-			const char* DefaultPrecision = bDefaultPrecisionIsHalf ? "mediump" : "highp";
-			ralloc_asprintf_append(buffer, "precision %s float;\n", DefaultPrecision);
-			// always use highp for integers as shaders use them as bit storage
-			ralloc_asprintf_append(buffer, "precision %s int;\n", "highp");
-		}
-
-		foreach_iter(exec_list_iterator, iter, *ir)
-		{
-			ir_instruction *inst = (ir_instruction *)iter.get();
-			do_visit(inst);
-		}
-		buffer = 0;
-
-		char* decl_buffer = ralloc_asprintf(mem_ctx, "");
-		buffer = &decl_buffer;
-		declare_structs(state, bCanHaveUBs);
-		buffer = 0;
-
-		char* signature = ralloc_asprintf(mem_ctx, "");
-		buffer = &signature;
-		print_signature(state, bGroupFlattenedUBs);
-		buffer = 0;
-
-		const char* geometry_layouts = "";
-		if (state->maxvertexcount>0)
-		{
-			check(state->geometryinput>0);
-			check(state->outputstream_type>0);
-			geometry_layouts = ralloc_asprintf(
-				mem_ctx,
-				"\nlayout(%s) in;\nlayout(%s, max_vertices = %u) out;\n\n",
-				GeometryInputStrings[state->geometryinput],
-				OutputStreamTypeStrings[state->outputstream_type],
-				state->maxvertexcount);
-		}
-		char* layout = ralloc_asprintf(mem_ctx, "");
-		buffer = &layout;
-		print_layout(state);
-		buffer = 0;
-
-		const FVulkanLanguageSpec* LanguageSpec = static_cast<const FVulkanLanguageSpec*>(state->LanguageSpec);
-
-		const bool bShouldEmitOESExtensions = LanguageSpec->RequiresOESExtensions();
-		const bool bShouldEmitMultiview = strstr(signature, "gl_ViewIndex") != nullptr;
-
-		char* Extensions = ralloc_asprintf(mem_ctx, "");
-		buffer = &Extensions;
-		print_extensions(state, state->language_version == 310, bShouldEmitOESExtensions, bShouldEmitMultiview);
-		if (state->bSeparateShaderObjects && !state->bGenerateES)
-		{
-			switch (state->target)
-			{
-			case geometry_shader:
-#if 0
-				ralloc_asprintf_append(buffer, "in gl_PerVertex\n"
-					"{\n"
-					"\tvec4 gl_Position;\n"
-					"\tfloat gl_ClipDistance[];\n"
-					"} gl_in[];\n"
-					);
-#endif
-				break;
-			case vertex_shader:
-#if 0
-				ralloc_asprintf_append(buffer, "out gl_PerVertex\n"
-					"{\n"
-					"\tvec4 gl_Position;\n"
-					"\tfloat gl_ClipDistance[];\n"
-					"};\n"
-					);
-#endif
-				break;
-			case tessellation_control_shader:
-				ralloc_asprintf_append(buffer, "in gl_PerVertex\n"
-					"{\n"
-					"\tvec4 gl_Position;\n"
-					"\tfloat gl_ClipDistance[];\n"
-					"} gl_in[gl_MaxPatchVertices];\n"
-					);
-				ralloc_asprintf_append(buffer, "out gl_PerVertex\n"
-					"{\n"
-					"\tvec4 gl_Position;\n"
-					"\tfloat gl_ClipDistance[];\n"
-					"} gl_out[];\n"
-					);
-				break;
-			case tessellation_evaluation_shader:
-				ralloc_asprintf_append(buffer, "in gl_PerVertex\n"
-					"{\n"
-					"\tvec4 gl_Position;\n"
-					"\tfloat gl_ClipDistance[];\n"
-					"} gl_in[gl_MaxPatchVertices];\n"
-					);
-				ralloc_asprintf_append(buffer, "out gl_PerVertex\n"
-					"{\n"
-					"\tvec4 gl_Position;\n"
-					"\tfloat gl_ClipDistance[];\n"
-					"};\n"
-					);
-				break;
-			case fragment_shader:
-			case compute_shader:
-			default:
-				break;
-			}
-		}
-
-		// Here since the code_buffer must have been populated beforehand
-		if (ParseState->LanguageSpec->AllowsSharingSamplers())
-		{
-			auto FindPrecision = [&](const char* Name)
-			{
-				for (auto& Pair : state->TextureToSamplerMap)
-				{
-					for (auto& Entry : Pair.second)
-					{
-						for (const std::string& SSName : SamplerStateNames)
-						{
-							if (Entry == SSName)
-							{
-								for (auto& PackedEntry : state->GlobalPackedArraysMap['s'])
-								{
-									if (!FCStringAnsi::Strcmp(Pair.first.c_str(), PackedEntry.Name.c_str()))
-									{
-										foreach_iter(exec_list_iterator, iter, sampler_variables)
-										{
-											ir_variable* var = ((extern_var*)iter.get())->var;
-											if (!FCStringAnsi::Strcmp(var->name, PackedEntry.CB_PackedSampler.c_str()))
-											{
-												return GetPrecisionModifierName(GetPrecisionModifier(var->type));
-											}
-										}
-									}
-								}
-							}
-						}
-					}
-				}
-
-				return "";
-			};
-			
-			BindingTable.InputAttachmentsMask = SubpassFetchMask;
-			for (uint32 Index = 0; SubpassFetchMask != 0; ++Index, SubpassFetchMask>>=1)
-			{
-				if ((SubpassFetchMask & 1) == 0)
-				{
-					continue;
-				}
-				
-				int32 BindingIndex = BindingTable.RegisterBinding(VULKAN_SUBPASS_FETCH_VAR[Index], "a", EVulkanBindingType::InputAttachment);
-				int32 InputAttachmentIndex = Index;
-				ralloc_asprintf_append(buffer, "layout(set=%d, binding=BINDING_%d, input_attachment_index=%d) uniform highp subpassInput %s;\n",
-					GetDescriptorSetForStage(ParseState->target),
-					BindingIndex,
-					InputAttachmentIndex,
-					VULKAN_SUBPASS_FETCH_VAR[Index]);
-
-				if (Index == 0) 
-				{
-					// Depth Input
-					ralloc_asprintf_append(buffer,
-						"highp float %s()\n"
-						"{\n"\
-						"\treturn subpassLoad(%s).x;\n"\
-						"}\n\n",
-						VULKAN_SUBPASS_FETCH[Index],
-						VULKAN_SUBPASS_FETCH_VAR[Index]);
-				}
-				else
-				{
-					// Color Input
-					ralloc_asprintf_append(buffer,
-						"highp vec4 %s()\n"
-						"{\n"\
-						"\treturn subpassLoad(%s);\n"\
-						"}\n\n",
-						VULKAN_SUBPASS_FETCH[Index],
-						VULKAN_SUBPASS_FETCH_VAR[Index]);
-				}
-			}
-						
-			for (int32 Index = 0; Index < BindingTable.Bindings.Num(); ++Index)
-			{
-				if (BindingTable.Bindings[Index].Type == EVulkanBindingType::Sampler)
-				{
-					const char* Precision = FindPrecision(BindingTable.Bindings[Index].Name);
-					ralloc_asprintf_append(buffer, "layout(set=%d, binding=BINDING_%d) uniform %s sampler %s;\n",
-						GetDescriptorSetForStage(ParseState->target),
-						Index,
-						Precision,
-						BindingTable.Bindings[Index].Name);
-				}
-			}
-		}
-
-		buffer = 0;
-
-		BindingTable.SortBindings();
-		char* BindingMapping = ralloc_asprintf(mem_ctx, "");
-		BindingTable.PrintBindingTableDefines(&BindingMapping);
-
-		if (state->target == vertex_shader || state->target == geometry_shader)
-		{
-			ralloc_asprintf_append(&BindingMapping, "\ninvariant gl_Position;\n");
-		}
-
-		const char* RequiredExtensions =
-		"#extension GL_ARB_separate_shader_objects : enable\n"
-		"#extension GL_ARB_shading_language_420pack : enable\n";
-
-		char* full_buffer = ralloc_asprintf(
-			state,
-			"// Compiled by HLSLCC %d.%d\n"		// HLSLCC_VersionMajor, HLSLCC_VersionMinor
-			"%s"								// signature
-			"#version %u %s\n"					// state->language_version, state->language_version == 310 ? "es" : ""
-			"%s"								// vulkan_required_extension
-			"%s"								// Bindings
-			"%s"								// Extensions
-			"%s"								// geometry_layouts
-			"%s"								// layout
-			"%s"								// decl_buffer
-			"%s"								// code_buffer
-			"\n",
-			HLSLCC_VersionMajor, HLSLCC_VersionMinor,
-			signature,
-			(Target == HCT_FeatureLevelSM4 || Target == HCT_FeatureLevelSM5) ? 430 : state->language_version,
-			state->language_version == 310 ? "es" : "",
-			state->language_version == 310 ? "" : RequiredExtensions,
-			BindingMapping,
-			Extensions,
-			geometry_layouts,
-			layout,
-			decl_buffer,
-			code_buffer
-			);
-		ralloc_free(mem_ctx);
-
-		return full_buffer;
-	}
-};
-
-struct FBreakPrecisionChangesVisitor : public ir_rvalue_visitor
-{
-	_mesa_glsl_parse_state* State;
-	const bool bDefaultPrecisionIsHalf;
-
-	FBreakPrecisionChangesVisitor(_mesa_glsl_parse_state* InState, bool bInDefaultPrecisionIsHalf) 
-		: State(InState)
-		, bDefaultPrecisionIsHalf(bInDefaultPrecisionIsHalf)
-	{}
-
-	virtual void handle_rvalue(ir_rvalue** RValuePtr) override
-	{
-		if (!RValuePtr || !*RValuePtr)
-		{
-			return;
-		}
-		bool bGenerateNewVar = false;
-		auto* RValue = *RValuePtr;
-		auto* Expression = RValue->as_expression();
-		auto* Constant = RValue->as_constant();
-		if (Expression)
-		{
-			if (bDefaultPrecisionIsHalf)
-			{
-				switch (Expression->operation)
-				{
-				case ir_unop_i2f:
-				case ir_unop_b2f:
-				case ir_unop_u2f:
-					bGenerateNewVar = bDefaultPrecisionIsHalf;
-					break;
-
-				case ir_unop_i2h:
-				case ir_unop_b2h:
-				case ir_unop_u2h:
-					bGenerateNewVar = !bDefaultPrecisionIsHalf;
-					break;
-
-				case ir_unop_h2f:
-				case ir_unop_f2h:
-					if (!Expression->operands[0]->as_texture())
-					{
-						bGenerateNewVar = true;
-					}
-					break;
-				}
-			}
-		}
-		else if (Constant)
-		{
-			/*
-			if ((bDefaultPrecisionIsHalf && Constant->type->base_type == GLSL_TYPE_HALF) ||
-			(!bDefaultPrecisionIsHalf && Constant->type->base_type == GLSL_TYPE_FLOAT))
-			{
-			bGenerateNewVar = true;
-			}
-			*/
-		}
-		if (bGenerateNewVar)
-		{
-			auto* NewVar = new(State)ir_variable(RValue->type, nullptr, ir_var_temporary);
-			auto* NewAssignment = new(State)ir_assignment(new(State)ir_dereference_variable(NewVar), RValue);
-			*RValuePtr = new(State)ir_dereference_variable(NewVar);
-			base_ir->insert_before(NewVar);
-			base_ir->insert_before(NewAssignment);
-		}
-	}
-};
-
-void FGenerateVulkanVisitor::AddTypeToUsedStructs(const glsl_type* type)
-{
-	if (type->base_type == GLSL_TYPE_STRUCT)
-	{
-		if (hash_table_find(used_structures, type) == NULL)
-		{
-			hash_table_insert(used_structures, (void*)type, type);
-		}
-	}
-
-	if (type->base_type == GLSL_TYPE_ARRAY && type->fields.array->base_type == GLSL_TYPE_STRUCT)
-	{
-		if (hash_table_find(used_structures, type->fields.array) == NULL)
-		{
-			hash_table_insert(used_structures, (void*)type->fields.array, type->fields.array);
-		}
-	}
-
-	if ((type->base_type == GLSL_TYPE_INPUTPATCH || type->base_type == GLSL_TYPE_OUTPUTPATCH) && type->inner_type->base_type == GLSL_TYPE_STRUCT)
-	{
-		if (hash_table_find(used_structures, type->inner_type) == NULL)
-		{
-			hash_table_insert(used_structures, (void*)type->inner_type, type->inner_type);
-		}
-	}
-}
-
-struct FGenerateSamplerToTextureMapVisitor : public ir_hierarchical_visitor
-{
-	_mesa_glsl_parse_state* State;
-	FSamplerMappingGatherData GatherData;
-
-	FGenerateSamplerToTextureMapVisitor(_mesa_glsl_parse_state* InState)
-		: State(InState)
-	{
-	}
-
-	virtual ir_visitor_status visit_enter(ir_texture* IR)
-	{
-		ir_variable* Sampler = IR->sampler->variable_referenced();
-		if (Sampler)
-		{
-			if (IR->SamplerStateName && IR->SamplerStateName[0])
-			{
-				// Keep in mind this sampler uses a sampling intrinsic, e.g. "Sample" or "SampleLevel" (Potential separate sampler sate).
-				GatherData.Entries[Sampler->name].bUsingSampling = true;
-				GatherData.Entries[Sampler->name].SamplerStates.insert(IR->SamplerStateName);
-				GatherData.SamplerToTextureMap[IR->SamplerStateName].insert(Sampler->name);
-			}
-			else
-			{
-				if (IR->op == ir_txf || IR->op == ir_txs || IR->op == ir_txm)
-				{
-					// Keep in mind this sampler uses the "Load" or "GetDimensions" intrinsic (Potential combined image-sampler).
-					GatherData.Entries[Sampler->name].bUsingLoadOrDim = true;
-				}
-				else
-				{
-#if UE_BUILD_DEBUG
-					// Internal error!!!
-					ensure(0);
-#endif
-					GatherData.Entries[Sampler->name].SamplerStates.insert("");
-				}
-			}
-		}
-		return ir_hierarchical_visitor::visit_enter(IR);
-	}
-};
-
-char* FVulkanCodeBackend::GenerateCode(exec_list* ir, _mesa_glsl_parse_state* state, EHlslShaderFrequency Frequency)
-{
-	FixRedundantCasts(ir);
-	//IRDump(ir, state);
-
-	FixIntrinsics(state, ir);
-
-	const bool bDefaultPrecisionIsHalf = (Frequency == HSF_PixelShader) && (HlslCompileFlags & HLSLCC_UseFullPrecisionInPS) == 0;
-
-	FBreakPrecisionChangesVisitor BreakPrecisionChangesVisitor(state, bDefaultPrecisionIsHalf);
-	BreakPrecisionChangesVisitor.run(ir);
-
-	const bool bGroupFlattenedUBs = ((HlslCompileFlags & HLSLCC_GroupFlattenedUniformBuffers) == HLSLCC_GroupFlattenedUniformBuffers);
-	const bool bGenerateLayoutLocations = state->bGenerateLayoutLocations;
-	const bool bCanHaveUBs = true;//(HlslCompileFlags & HLSLCC_FlattenUniformBuffers) != HLSLCC_FlattenUniformBuffers;
-
-	//IRDump(ir);
-
-	// Setup root visitor
-	FGenerateVulkanVisitor visitor(Target, BindingTable, state, bGenerateLayoutLocations, bDefaultPrecisionIsHalf);
-	visitor.FindAtomicVariables(ir);
-
-	// Generate information for sharing samplers
-	{
-		FGenerateSamplerToTextureMapVisitor GenerateSamplerToTextureMapVisitor(state);
-		GenerateSamplerToTextureMapVisitor.run(ir);
-		visitor.SamplerMapping.Consolidate(GenerateSamplerToTextureMapVisitor.GatherData);
-	}
-
-	uint32 SubpassFetchMask = 0;
-	if (Frequency == HSF_PixelShader)
-	{
-		for (int32 i = 0; i < UE_ARRAY_COUNT(VULKAN_SUBPASS_FETCH); ++i)	
-		{
-			SubpassFetchMask|= (UsesUEIntrinsic(ir, VULKAN_SUBPASS_FETCH[i]) ? 1 << i : 0);
-		}
-	}
-
-	const char* code = visitor.run(ir, state, bGroupFlattenedUBs, bCanHaveUBs, SubpassFetchMask);
-
-	return _strdup(code);
-}
-
-
-// Verify if SampleLevel() is used
-struct SPromoteSampleLevelES : public ir_hierarchical_visitor
-{
-	_mesa_glsl_parse_state* ParseState;
-	const bool bIsVertexShader;
-	SPromoteSampleLevelES(_mesa_glsl_parse_state* InParseState, bool bInIsVertexShader) :
-		ParseState(InParseState),
-		bIsVertexShader(bInIsVertexShader)
-	{
-	}
-
-	virtual ir_visitor_status visit_leave(ir_texture* IR) override
-	{
-		if (IR->op == ir_txl)
-		{
-			if (bIsVertexShader)
-			{
-				YYLTYPE loc;
-				loc.first_column = IR->SourceLocation.Column;
-				loc.first_line = IR->SourceLocation.Line;
-				loc.source_file = IR->SourceLocation.SourceFile;
-				_mesa_glsl_error(&loc, ParseState, "Vertex texture fetch currently not supported on GLSL ES\n");
-			}
-			else
-			{
-				//@todo-mobile: allowing lod texture functions for now, as they are supported on some devices via glsl extension.
-				// http://www.khronos.org/registry/gles/extensions/EXT/EXT_shader_texture_lod.txt
-				// Compat work will be required for devices which do not support it.
-				/*
-				_mesa_glsl_warning(ParseState, "%s(%u, %u) Converting SampleLevel() to Sample()\n", IR->SourceLocation.SourceFile.c_str(), IR->SourceLocation.Line, IR->SourceLocation.Column);
-				IR->op = ir_tex;
-				*/
-			}
-		}
-
-		if (IR->offset)
-		{
-			YYLTYPE loc;
-			loc.first_column = IR->SourceLocation.Column;
-			loc.first_line = IR->SourceLocation.Line;
-			loc.source_file = IR->SourceLocation.SourceFile;
-			_mesa_glsl_error(&loc, ParseState, "Texture offset not supported on GLSL ES\n");
-		}
-
-		return visit_continue;
-	}
-};
-
-
-// Converts an array index expression using an integer input attribute, to a float input attribute using a conversion to int
-struct SConvertIntVertexAttributeES final : public ir_hierarchical_visitor
-{
-	_mesa_glsl_parse_state* ParseState;
-	exec_list* FunctionBody;
-	int InsideArrayDeref;
-	std::map<ir_variable*, ir_variable*> ConvertedVarMap;
-
-	SConvertIntVertexAttributeES(_mesa_glsl_parse_state* InParseState, exec_list* InFunctionBody) : ParseState(InParseState), FunctionBody(InFunctionBody), InsideArrayDeref(0)
-	{
-	}
-
-	virtual ~SConvertIntVertexAttributeES()
-	{
-	}
-
-	virtual ir_visitor_status visit_enter(ir_dereference_array* DeRefArray) override
-	{
-		// Break the array dereference so we know we want to modify the array index part
-		auto Result = ir_hierarchical_visitor::visit_enter(DeRefArray);
-		++InsideArrayDeref;
-		DeRefArray->array_index->accept(this);
-		--InsideArrayDeref;
-
-		return visit_continue;
-	}
-
-	virtual ir_visitor_status visit(ir_dereference_variable* DeRefVar) override
-	{
-		if (InsideArrayDeref > 0)
-		{
-			ir_variable* SourceVar = DeRefVar->var;
-			if (SourceVar->mode == ir_var_in)
-			{
-				// First time it still is an integer, so add the temporary and a conversion, and switch to float
-				if (SourceVar->type->is_integer())
-				{
-					check(SourceVar->type->is_integer() && !SourceVar->type->is_matrix() && !SourceVar->type->is_array());
-
-					// Double check we haven't processed this
-					auto IterFound = ConvertedVarMap.find(SourceVar);
-					check(IterFound == ConvertedVarMap.end());
-
-					// New temp var
-					ir_variable* NewVar = new(ParseState)ir_variable(SourceVar->type, NULL, ir_var_temporary);
-					base_ir->insert_before(NewVar);
-
-					// Switch original type to float
-					SourceVar->type = glsl_type::get_instance(GLSL_TYPE_FLOAT, SourceVar->type->vector_elements, 1);
-
-					// Convert float to int
-					ir_dereference_variable* NewSourceDeref = new(ParseState)ir_dereference_variable(SourceVar);
-					ir_expression* NewCastExpression = new(ParseState)ir_expression(ir_unop_f2i, NewSourceDeref);
-					ir_assignment* NewAssigment = new(ParseState)ir_assignment(new(ParseState)ir_dereference_variable(NewVar), NewCastExpression);
-					base_ir->insert_before(NewAssigment);
-
-					// Add the entry and modify the original Var
-					ConvertedVarMap[SourceVar] = NewVar;
-					DeRefVar->var = NewVar;
-				}
-				else
-				{
-					auto IterFound = ConvertedVarMap.find(SourceVar);
-					if (IterFound != ConvertedVarMap.end())
-					{
-						DeRefVar->var = IterFound->second;
-					}
-				}
-			}
-		}
-
-		return ir_hierarchical_visitor::visit(DeRefVar);
-	}
-};
-
-
-bool FVulkanCodeBackend::ApplyAndVerifyPlatformRestrictions(exec_list* Instructions, _mesa_glsl_parse_state* ParseState, EHlslShaderFrequency Frequency)
-{
-	if (ParseState->bGenerateES)
-	{
-		bool bIsVertexShader = (Frequency == HSF_VertexShader);
-
-		// Handle SampleLevel
-		{
-			SPromoteSampleLevelES Visitor(ParseState, bIsVertexShader);
-			Visitor.run(Instructions);
-		}
-
-		// Handle matrices (flatten to vectors so we can support non-sqaure)
-		ExpandMatricesIntoArrays(Instructions, ParseState);
-
-		// Handle integer vertex attributes used as array indices
-		if (bIsVertexShader)
-		{
-			SConvertIntVertexAttributeES ConvertIntVertexAttributeVisitor(ParseState, Instructions);
-			ConvertIntVertexAttributeVisitor.run(Instructions);
-		}
-	}
-
-	return true;
-}
-
-/** Qualifers that apply to semantics. */
-union FSemanticQualifier
-{
-	struct
-	{
-		unsigned bCentroid : 1;
-		unsigned InterpolationMode : 2;
-		unsigned bIsPatchConstant : 1;
-	} Fields;
-	unsigned Packed;
-
-	FSemanticQualifier() : Packed(0) {}
-};
-
-/** Information on system values. */
-struct FSystemValue
-{
-	const char* Semantic;
-	const glsl_type* Type;
-	const char* GlslName;
-	ir_variable_mode Mode;
-	bool bOriginUpperLeft;
-	bool bArrayVariable;
-	bool bApplyClipSpaceAdjustment;
-	bool bESOnly;
-};
-
-/** Vertex shader system values. */
-static FSystemValue VertexSystemValueTable[] =
-{
-	{ "SV_ViewID", glsl_type::int_type, "gl_ViewIndex", ir_var_in, false, false, false, false },
-	{ "SV_VertexID", glsl_type::int_type, "gl_VertexIndex", ir_var_in, false, false, false, false },
-	{ "SV_InstanceID", glsl_type::int_type, "gl_InstanceIndex", ir_var_in, false, false, false, false },
-	{ "SV_Position", glsl_type::vec4_type, "gl_Position", ir_var_out, false, false, true, false },
-	{ NULL, NULL, NULL, ir_var_auto, false, false, false, false }
-};
-
-/** Pixel shader system values. */
-static FSystemValue PixelSystemValueTable[] =
-{
-	{ "SV_ViewID", glsl_type::int_type, "gl_ViewIndex", ir_var_in, false, false, false, false },
-	{ "SV_Depth", glsl_type::float_type, "gl_FragDepth", ir_var_out, false, false, false, false },
-	{ "SV_Position", glsl_type::vec4_type, "gl_FragCoord", ir_var_in, true, false, false, false },
-	{ "SV_IsFrontFace", glsl_type::bool_type, "gl_FrontFacing", ir_var_in, false, false, true, false },
-	{ "SV_PrimitiveID", glsl_type::int_type, "gl_PrimitiveID", ir_var_in, false, false, false, false },
-	{ "SV_RenderTargetArrayIndex", glsl_type::int_type, "gl_Layer", ir_var_in, false, false, false, false },
-	{ "SV_Target0", glsl_type::half4_type, "gl_FragColor", ir_var_out, false, false, false, true },
-	{ "SV_Coverage", glsl_type::int_type, "gl_SampleMaskIn[0]", ir_var_in, false, false, false, false },
-	{ "SV_Coverage", glsl_type::int_type, "gl_SampleMask[0]", ir_var_out, false, false, false, false },
-	{ "SV_SampleIndex", glsl_type::uint_type, "gl_SampleID", ir_var_in, false, false, false, false},
-	{ NULL, NULL, NULL, ir_var_auto, false, false, false }
-};
-
-/** Geometry shader system values. */
-static FSystemValue GeometrySystemValueTable[] =
-{
-	{ "SV_ViewID", glsl_type::int_type, "gl_ViewIndex", ir_var_in, false, false, false, false },
-	{ "SV_VertexID", glsl_type::int_type, "gl_VertexID", ir_var_in, false, false, false, false },
-	{ "SV_InstanceID", glsl_type::int_type, "gl_InstanceID", ir_var_in, false, false, false, false },
-	{ "SV_Position", glsl_type::vec4_type, "gl_Position", ir_var_in, false, true, true, false },
-	{ "SV_Position", glsl_type::vec4_type, "gl_Position", ir_var_out, false, false, true, false },
-	{ "SV_RenderTargetArrayIndex", glsl_type::int_type, "gl_Layer", ir_var_out, false, false, false, false },
-	{ "SV_PrimitiveID", glsl_type::int_type, "gl_PrimitiveID", ir_var_out, false, false, false, false },
-	{ "SV_PrimitiveID", glsl_type::int_type, "gl_PrimitiveIDIn", ir_var_in, false, false, false, false },
-	{ NULL, NULL, NULL, ir_var_auto, false, false, false, false }
-};
-
-
-/** Hull shader system values. */
-static FSystemValue HullSystemValueTable[] =
-{
-	{ "SV_ViewID", glsl_type::int_type, "gl_ViewIndex", ir_var_in, false, false, false, false },
-	{ "SV_OutputControlPointID", glsl_type::int_type, "gl_InvocationID", ir_var_in, false, false, false, false },
-	{ NULL, NULL, NULL, ir_var_auto, false, false, false, false }
-};
-
-/** Domain shader system values. */
-static FSystemValue DomainSystemValueTable[] =
-{
-	{ "SV_ViewID", glsl_type::int_type, "gl_ViewIndex", ir_var_in, false, false, false, false },
-	// TODO : SV_DomainLocation has types float2 or float3 depending on the input topology
-	{ "SV_Position", glsl_type::vec4_type, "gl_Position", ir_var_in, false, true, true, false },
-	{ "SV_Position", glsl_type::vec4_type, "gl_Position", ir_var_out, false, false, true, false },
-	{ "SV_DomainLocation", glsl_type::vec3_type, "gl_TessCoord", ir_var_in, false, false, false, false },
-	{ NULL, NULL, NULL, ir_var_auto, false, false, false, false }
-};
-
-/** Compute shader system values. */
-static FSystemValue ComputeSystemValueTable[] =
-{
-	{ "SV_DispatchThreadID", glsl_type::uvec3_type, "gl_GlobalInvocationID", ir_var_in, false, false, false, false },
-	{ "SV_GroupID", glsl_type::uvec3_type, "gl_WorkGroupID", ir_var_in, false, false, false, false },
-	{ "SV_GroupIndex", glsl_type::uint_type, "gl_LocalInvocationIndex", ir_var_in, false, false, false, false },
-	{ "SV_GroupThreadID", glsl_type::uvec3_type, "gl_LocalInvocationID", ir_var_in, false, false, false, false },
-	{ NULL, NULL, NULL, ir_var_auto, false, false, false, false }
-};
-
-static FSystemValue* SystemValueTable[HSF_FrequencyCount] =
-{
-	VertexSystemValueTable,
-	PixelSystemValueTable,
-	GeometrySystemValueTable,
-	HullSystemValueTable,
-	DomainSystemValueTable,
-	ComputeSystemValueTable
-};
-
-
-#define CUSTOM_LAYER_INDEX_SEMANTIC "HLSLCC_LAYER_INDEX"
-
-// Returns the number of slots an in/out variable occupies. This is determined by the variable type, e.g. vec4 has offset 1, mat4 has offset 4.
-// For geometry shader input variables, the first array specifier is ignored as geometry shaders get their input from multiple vertex shader invocations.
-static unsigned GetInOutVariableSlotSize(EHlslShaderFrequency Frequency, ir_variable_mode Mode, const struct glsl_type* VarType)
-{
-	check(VarType != nullptr);
-	if (VarType->fields.array != nullptr)
-	{
-		if (Frequency == HSF_GeometryShader && Mode == ir_var_in)
-		{
-			// Ignore frequency after first array specifier, so use HSF_InvalidFrequency
-			return GetInOutVariableSlotSize(HSF_InvalidFrequency, Mode, VarType->fields.array);
-		}
-		else
-		{
-			// Multiply inner array type by array size, i.e. both 'vec2[4]' and 'float[4]' occupy 4 slots
-			// Ignore frequency after first array specifier, so use HSF_InvalidFrequency
-			return GetInOutVariableSlotSize(HSF_InvalidFrequency, Mode, VarType->fields.array) * VarType->length;
-		}
-	}
-	else
-	{
-		// Only use number of matrix columns to determine number of slots this variable occupies: 1 for scalars and vectors, 2/3/4 for matrices
-		return VarType->matrix_columns;
-	}
-}
-
-static void ConfigureInOutVariableLayout(EHlslShaderFrequency Frequency,
-	_mesa_glsl_parse_state* ParseState,
-	const char* Semantic,
-	ir_variable* Variable,
-	ir_variable_mode Mode
-	)
-{
-	if (Frequency == HSF_VertexShader && Mode == ir_var_in)
-	{
-		const int PrefixLength = 9;
-		if ((FCStringAnsi::Strnicmp(Semantic, "ATTRIBUTE", PrefixLength) == 0) &&
-			(Semantic[PrefixLength] >= '0') && (Semantic[PrefixLength] <= '9')
-			)
-		{
-			int AttributeIndex = atoi(Semantic + PrefixLength);
-
-			Variable->explicit_location = true;
-			Variable->location = AttributeIndex;
-			Variable->semantic = ralloc_strdup(Variable, Semantic);
-		}
-		else
-		{
-#ifdef DEBUG
-#define _mesh_glsl_report _mesa_glsl_warning
-#else
-#define _mesh_glsl_report _mesa_glsl_error
-#endif
-			_mesh_glsl_report(ParseState, "Vertex shader input semantic must be ATTRIBUTE and not \'%s\' in order to determine location/semantic index", Semantic);
-#undef _mesh_glsl_report
-		}
-	}
-	else if (Semantic && FCStringAnsi::Strnicmp(Variable->name, "gl_", 3) != 0)
-	{
-		Variable->explicit_location = 1;
-		Variable->semantic = ralloc_strdup(Variable, Semantic);
-		unsigned int NumVectors = (Variable->type->matrix_columns > 1) ? Variable->type->matrix_columns : 1;
-		if (Mode == ir_var_in)
-		{
-			Variable->location = ParseState->next_in_location_slot;
-			ParseState->next_in_location_slot += GetInOutVariableSlotSize(Frequency, Mode, Variable->type);
-		}
-		else
-		{
-			Variable->location = ParseState->next_out_location_slot;
-			ParseState->next_out_location_slot += GetInOutVariableSlotSize(Frequency, Mode, Variable->type);
-		}
-	}
-}
-
-/**
-* Generate an input semantic.
-* @param Frequency - The shader frequency.
-* @param ParseState - Parse state.
-* @param Semantic - The semantic name to generate.
-* @param Type - Value type.
-* @param DeclInstructions - IR to which declarations may be added.
-* @returns reference to IR variable for the semantic.
-*/
-static ir_rvalue* GenShaderInputSemantic(
-	EHlslShaderFrequency Frequency,
-	_mesa_glsl_parse_state* ParseState,
-	const char* Semantic,
-	FSemanticQualifier InputQualifier,
-	const glsl_type* Type,
-	exec_list* DeclInstructions,
-	int SemanticArraySize,
-	int SemanticArrayIndex,
-	bool& ApplyClipSpaceAdjustment
-	)
-{
-	if (Semantic && FCStringAnsi::Strnicmp(Semantic, "SV_", 3) == 0)
-	{
-		FSystemValue* SystemValues = SystemValueTable[Frequency];
-		for (int i = 0; SystemValues[i].Semantic != NULL; ++i)
-		{
-			if (SystemValues[i].Mode == ir_var_in
-				&& (!SystemValues[i].bESOnly || ParseState->bGenerateES)
-				&& FCStringAnsi::Stricmp(SystemValues[i].Semantic, Semantic) == 0)
-			{
-				if (SystemValues[i].bArrayVariable)
-				{
-					// Built-in array variable. Like gl_in[x].gl_Position.
-					// The variable for it has already been created in GenShaderInput().
-					ir_variable* Variable = ParseState->symbols->get_variable("gl_in");
-					check(Variable);
-					ir_dereference_variable* ArrayDeref = new(ParseState)ir_dereference_variable(Variable);
-					ir_dereference_array* StructDeref = new(ParseState)ir_dereference_array(
-						ArrayDeref,
-						new(ParseState)ir_constant((unsigned)SemanticArrayIndex)
-						);
-					ir_dereference_record* VariableDeref = new(ParseState)ir_dereference_record(
-						StructDeref,
-						SystemValues[i].GlslName
-						);
-					ApplyClipSpaceAdjustment = SystemValues[i].bApplyClipSpaceAdjustment;
-					// TO DO - in case of SV_ClipDistance, we need to defer appropriate index in variable too.
-					return VariableDeref;
-				}
-				else
-				{
-					// Built-in variable that shows up only once, like gl_FragCoord in fragment
-					// shader, or gl_PrimitiveIDIn in geometry shader. Unlike gl_in[x].gl_Position.
-					// Even in geometry shader input pass it shows up only once.
-
-					// Create it on first pass, ignore the call on others.
-					if (SemanticArrayIndex == 0)
-					{
-						ir_variable* Variable = new(ParseState)ir_variable(
-							SystemValues[i].Type,
-							SystemValues[i].GlslName,
-							ir_var_in
-							);
-						Variable->read_only = true;
-						Variable->origin_upper_left = SystemValues[i].bOriginUpperLeft;
-						DeclInstructions->push_tail(Variable);
-						ParseState->symbols->add_variable(Variable);
-						ir_dereference_variable* VariableDeref = new(ParseState)ir_dereference_variable(Variable);
-
-						if (FCStringAnsi::Stricmp(Semantic, "SV_Position") == 0 && Frequency == HSF_PixelShader)
-						{
-							// This is for input of gl_FragCoord into pixel shader only.
-
-							// Generate a local variable to do the conversion in, keeping source type.
-							ir_variable* TempVariable = new(ParseState)ir_variable(Variable->type, NULL, ir_var_temporary);
-							DeclInstructions->push_tail(TempVariable);
-
-							// Assign input to this variable
-							ir_dereference_variable* TempVariableDeref = new(ParseState)ir_dereference_variable(TempVariable);
-							DeclInstructions->push_tail(
-								new(ParseState)ir_assignment(
-								TempVariableDeref,
-								VariableDeref
-								)
-								);
-
-							// TempVariable.w = ( 1.0f / TempVariable.w );
-							DeclInstructions->push_tail(
-								new(ParseState)ir_assignment(
-								new(ParseState)ir_swizzle(TempVariableDeref->clone(ParseState, NULL), 3, 0, 0, 0, 1),
-								new(ParseState)ir_expression(ir_binop_div,
-								new(ParseState)ir_constant(1.0f),
-								new(ParseState)ir_swizzle(TempVariableDeref->clone(ParseState, NULL), 3, 0, 0, 0, 1)
-								)
-								)
-								);
-
-							return TempVariableDeref->clone(ParseState, NULL);
-						}
-						else if (SystemValues[i].bApplyClipSpaceAdjustment)
-						{
-							// incoming gl_FrontFacing. Make it (!gl_FrontFacing), due to vertical flip in OpenGL
-							return new(ParseState)ir_expression(ir_unop_logic_not, glsl_type::bool_type, VariableDeref, NULL);
-						}
-						else
-						{
-							return VariableDeref;
-						}
-					}
-					else
-					{
-						return NULL;
-					}
-				}
-			}
-		}
-	}
-
-	ir_variable* Variable = NULL;
-	if (Variable == NULL && Frequency == HSF_DomainShader)
-	{
-		const int PrefixLength = 13;
-		if (FCStringAnsi::Strnicmp(Semantic, "SV_TessFactor", PrefixLength) == 0
-			&& Semantic[PrefixLength] >= '0'
-			&& Semantic[PrefixLength] <= '3')
-		{
-			int OutputIndex = Semantic[PrefixLength] - '0';
-			Variable = new(ParseState)ir_variable(
-				Type,
-				ralloc_asprintf(ParseState, "gl_TessLevelOuter[%d]", OutputIndex),
-				ir_var_out
-				);
-		}
-	}
-
-	if (Variable == NULL && Frequency == HSF_DomainShader)
-	{
-		const int PrefixLength = 19;
-		if (FCStringAnsi::Strnicmp(Semantic, "SV_InsideTessFactor", PrefixLength) == 0
-			&& Semantic[PrefixLength] >= '0'
-			&& Semantic[PrefixLength] <= '1')
-		{
-			int OutputIndex = Semantic[PrefixLength] - '0';
-			Variable = new(ParseState)ir_variable(
-				Type,
-				ralloc_asprintf(ParseState, "gl_TessLevelInner[%d]", OutputIndex),
-				ir_var_out
-				);
-		}
-		else if (FCStringAnsi::Stricmp(Semantic, "SV_InsideTessFactor") == 0)
-		{
-			Variable = new(ParseState)ir_variable(
-				Type,
-				ralloc_asprintf(ParseState, "gl_TessLevelInner[0]"),
-				ir_var_out
-				);
-		}
-	}
-
-	if (Variable == NULL && (Frequency == HSF_VertexShader || Frequency == HSF_PixelShader))
-	{
-		if (FCStringAnsi::Stricmp(Semantic, "SV_ViewId") == 0)
-		{
-			Variable = new(ParseState)ir_variable(
-				Type,
-				ralloc_asprintf(ParseState, "gl_ViewIndex"),
-				ir_var_in
-			);
-		}
-	}
-
-	if (Variable)
-	{
-		// Up to this point, variables aren't contained in structs
-		DeclInstructions->push_tail(Variable);
-		ParseState->symbols->add_variable(Variable);
-		Variable->centroid = InputQualifier.Fields.bCentroid;
-		Variable->interpolation = InputQualifier.Fields.InterpolationMode;
-		Variable->is_patch_constant = InputQualifier.Fields.bIsPatchConstant;
-		ir_rvalue* VariableDeref = new(ParseState)ir_dereference_variable(Variable);
-
-		return VariableDeref;
-	}
-
-	// If we're here, no built-in variables matched.
-
-	if (Semantic && FCStringAnsi::Strnicmp(Semantic, "SV_", 3) == 0)
-	{
-		_mesa_glsl_warning(ParseState, "unrecognized system "
-			"value input '%s'", Semantic);
-	}
-
-	// Patch constants must be variables, not structs or interface blocks, in GLSL <= 4.10
-	bool bUseGLSL410Rules = InputQualifier.Fields.bIsPatchConstant && ParseState->language_version <= 410;
-	if (Frequency == HSF_VertexShader || ParseState->bGenerateES || bUseGLSL410Rules)
-	{
-		const char* Prefix = "in";
-		if ((ParseState->bGenerateES && Frequency == HSF_PixelShader) || bUseGLSL410Rules)
-		{
-			Prefix = "var";
-		}
-
-		// Vertex shader inputs don't get packed into structs that we'll later morph into interface blocks
-		if (ParseState->bGenerateES && Type->is_integer())
-		{
-			// Convert integer attributes to floats
-			Variable = new(ParseState)ir_variable(
-				Type,
-				ralloc_asprintf(ParseState, "%s_%s_I", Prefix, Semantic),
-				ir_var_temporary
-				);
-			Variable->centroid = InputQualifier.Fields.bCentroid;
-			Variable->interpolation = InputQualifier.Fields.InterpolationMode;
-			check(Type->is_vector() || Type->is_scalar());
-			check(Type->base_type == GLSL_TYPE_INT || Type->base_type == GLSL_TYPE_UINT);
-
-			// New float attribute
-			ir_variable* ReplacedAttributeVar = new (ParseState)ir_variable(glsl_type::get_instance(GLSL_TYPE_FLOAT, Variable->type->vector_elements, 1), ralloc_asprintf(ParseState, "%s_%s", Prefix, Semantic), ir_var_in);
-			ReplacedAttributeVar->read_only = true;
-			ReplacedAttributeVar->centroid = InputQualifier.Fields.bCentroid;
-			ReplacedAttributeVar->interpolation = InputQualifier.Fields.InterpolationMode;
-
-			// Convert to integer
-			ir_assignment* ConversionAssignment = new(ParseState)ir_assignment(
-				new(ParseState)ir_dereference_variable(Variable),
-				new(ParseState)ir_expression(
-				Type->base_type == GLSL_TYPE_INT ? ir_unop_f2i : ir_unop_f2u,
-				new (ParseState)ir_dereference_variable(ReplacedAttributeVar)
-				)
-				);
-
-			DeclInstructions->push_tail(ReplacedAttributeVar);
-			DeclInstructions->push_tail(Variable);
-			DeclInstructions->push_tail(ConversionAssignment);
-			ParseState->symbols->add_variable(Variable);
-			ParseState->symbols->add_variable(ReplacedAttributeVar);
-
-			ir_dereference_variable* VariableDeref = new(ParseState)ir_dereference_variable(ReplacedAttributeVar);
-			return VariableDeref;
-		}
-
-		// Regular attribute
-		Variable = new(ParseState)ir_variable(
-			Type,
-			ralloc_asprintf(ParseState, "%s_%s", Prefix, Semantic),
-			ir_var_in
-			);
-		Variable->read_only = true;
-		Variable->centroid = InputQualifier.Fields.bCentroid;
-		Variable->interpolation = InputQualifier.Fields.InterpolationMode;
-		Variable->is_patch_constant = InputQualifier.Fields.bIsPatchConstant;
-
-		if (ParseState->bGenerateLayoutLocations)
-		{
-			ConfigureInOutVariableLayout(Frequency, ParseState, Semantic, Variable, ir_var_in);
-		}
-
-		DeclInstructions->push_tail(Variable);
-		ParseState->symbols->add_variable(Variable);
-
-		ir_dereference_variable* VariableDeref = new(ParseState)ir_dereference_variable(Variable);
-		return VariableDeref;
-	}
-	else if (SemanticArrayIndex == 0)
-	{
-		// This code-section replaces "layout(location=0) in struct { vec4 Data; } in_ATTRIBUTE0;" pattern to
-		// "layout(location=0) in vec4 in_ATTRIBUTE0;".
-
-		if (/*Frequency == HSF_GeometryShader && */SemanticArraySize != 0)
-		{
-			Type = glsl_type::get_array_instance(Type, SemanticArraySize);
-		}
-
-		// Regular attribute
-		Variable = new(ParseState)ir_variable(
-			Type,
-			ralloc_asprintf(ParseState, "in_%s", Semantic),
-			ir_var_in
-			);
-		Variable->read_only = true;
-		Variable->centroid = InputQualifier.Fields.bCentroid;
-		Variable->interpolation = InputQualifier.Fields.InterpolationMode;
-		Variable->is_patch_constant = InputQualifier.Fields.bIsPatchConstant;
-
-		if (ParseState->bGenerateLayoutLocations)
-		{
-			ConfigureInOutVariableLayout(Frequency, ParseState, Semantic, Variable, ir_var_in);
-		}
-
-		DeclInstructions->push_tail(Variable);
-		ParseState->symbols->add_variable(Variable);
-
-		ir_dereference* VariableDeref = new(ParseState)ir_dereference_variable(Variable);
-		if (SemanticArraySize > 0)
-		{
-			// Deref inside array first
-			VariableDeref = new(ParseState) ir_dereference_array(VariableDeref, new(ParseState) ir_constant((unsigned)SemanticArrayIndex));
-		}
-
-		return VariableDeref;
-	}
-	else
-	{
-		// Array variable, not first pass. It already exists, get it.
-		Variable = ParseState->symbols->get_variable(ralloc_asprintf(ParseState, "in_%s", Semantic));
-		check(Variable);
-
-		ir_rvalue* VariableDeref = new(ParseState)ir_dereference_variable(Variable);
-		VariableDeref = new(ParseState)ir_dereference_array(VariableDeref, new(ParseState)ir_constant((unsigned)SemanticArrayIndex));
-		//VariableDeref = new(ParseState)ir_dereference_record(VariableDeref, ralloc_strdup(ParseState, "Data"));
-		return VariableDeref;
-	}
-}
-
-/**
-* Generate an output semantic.
-* @param Frequency - The shader frequency.
-* @param ParseState - Parse state.
-* @param Semantic - The semantic name to generate.
-* @param Type - Value type.
-* @param DeclInstructions - IR to which declarations may be added.
-* @returns the IR variable for the semantic.
-*/
-static ir_rvalue* GenShaderOutputSemantic(
-	EHlslShaderFrequency Frequency,
-	_mesa_glsl_parse_state* ParseState,
-	const char* Semantic,
-	FSemanticQualifier OutputQualifier,
-	const glsl_type* Type,
-	exec_list* DeclInstructions,
-	const glsl_type** DestVariableType,
-	bool& ApplyClipSpaceAdjustment,
-	bool& ApplyClampPowerOfTwo
-	)
-{
-	check(Semantic);
-
-	FSystemValue* SystemValues = SystemValueTable[Frequency];
-	ir_variable* Variable = NULL;
-
-	if (FCStringAnsi::Strnicmp(Semantic, "SV_", 3) == 0)
-	{
-		for (int i = 0; SystemValues[i].Semantic != NULL; ++i)
-		{
-			if (!SystemValues[i].bESOnly || ParseState->bGenerateES)
-			{
-				if (SystemValues[i].Mode == ir_var_out
-					&& FCStringAnsi::Stricmp(SystemValues[i].Semantic, Semantic) == 0)
-				{
-					Variable = new(ParseState)ir_variable(
-						SystemValues[i].Type,
-						SystemValues[i].GlslName,
-						ir_var_out
-						);
-					Variable->origin_upper_left = SystemValues[i].bOriginUpperLeft;
-					ApplyClipSpaceAdjustment = SystemValues[i].bApplyClipSpaceAdjustment;
-				}
-			}
-		}
-	}
-
-	if (Variable == NULL && (Frequency == HSF_VertexShader || Frequency == HSF_GeometryShader || Frequency == HSF_HullShader || Frequency == HSF_DomainShader))
-	{
-		const int PrefixLength = 15;
-		// Match SV_ClipDistance or SV_ClipDistanceN
-		if (FCStringAnsi::Strnicmp(Semantic, "SV_ClipDistance", PrefixLength) == 0
-			&& ((Semantic[PrefixLength] >= '0' && Semantic[PrefixLength] <= '9') || Semantic[PrefixLength] == 0))
-		{
-			int OutputIndex = Semantic[PrefixLength] ? Semantic[PrefixLength] - '0' : 0;
-			Variable = new(ParseState)ir_variable(
-				glsl_type::float_type,
-				ralloc_asprintf(ParseState, "gl_ClipDistance[%d]", OutputIndex),
-				ir_var_out
-				);
-		}
-	}
-
-	if (Variable == NULL && Frequency == HSF_PixelShader)
-	{
-		const int PrefixLength = 9;
-		if (FCStringAnsi::Strnicmp(Semantic, "SV_Target", PrefixLength) == 0
-			&& Semantic[PrefixLength] >= '0'
-			&& Semantic[PrefixLength] <= '7')
-		{
-			int OutputIndex = Semantic[PrefixLength] - '0';
-			Variable = new(ParseState)ir_variable(
-				Type,
-				ralloc_asprintf(ParseState, "out_Target%d", OutputIndex),
-				ir_var_out
-				);
-
-			if (ParseState->bGenerateLayoutLocations)
-			{
-				Variable->explicit_location = true;
-				Variable->location = OutputIndex;
-			}
-		}
-	}
-
-	if (Variable == NULL && Frequency == HSF_HullShader)
-	{
-		const int PrefixLength = 13;
-		if (FCStringAnsi::Strnicmp(Semantic, "SV_TessFactor", PrefixLength) == 0
-			&& Semantic[PrefixLength] >= '0'
-			&& Semantic[PrefixLength] <= '3')
-		{
-			int OutputIndex = Semantic[PrefixLength] - '0';
-			Variable = new(ParseState)ir_variable(
-				Type,
-				ralloc_asprintf(ParseState, "gl_TessLevelOuter[%d]", OutputIndex),
-				ir_var_out
-				);
-
-			ApplyClampPowerOfTwo = ParseState->tessellation.partitioning == GLSL_PARTITIONING_POW2;
-		}
-	}
-
-	if (Variable == NULL && Frequency == HSF_HullShader)
-	{
-		const int PrefixLength = 19;
-		if (FCStringAnsi::Strnicmp(Semantic, "SV_InsideTessFactor", PrefixLength) == 0
-			&& Semantic[PrefixLength] >= '0'
-			&& Semantic[PrefixLength] <= '1')
-		{
-			int OutputIndex = Semantic[PrefixLength] - '0';
-			Variable = new(ParseState)ir_variable(
-				Type,
-				ralloc_asprintf(ParseState, "gl_TessLevelInner[%d]", OutputIndex),
-				ir_var_out
-				);
-
-			ApplyClampPowerOfTwo = ParseState->tessellation.partitioning == GLSL_PARTITIONING_POW2;
-		}
-		else if (FCStringAnsi::Stricmp(Semantic, "SV_InsideTessFactor") == 0)
-		{
-			Variable = new(ParseState)ir_variable(
-				Type,
-				ralloc_asprintf(ParseState, "gl_TessLevelInner[0]"),
-				ir_var_out
-				);
-
-			ApplyClampPowerOfTwo = ParseState->tessellation.partitioning == GLSL_PARTITIONING_POW2;
-		}
-	}
-
-	bool bUseGLSL410Rules = OutputQualifier.Fields.bIsPatchConstant && ParseState->language_version == 410;
-	if (Variable == NULL && (ParseState->bGenerateES || bUseGLSL410Rules))
-	{
-		// Create a variable so that a struct will not get added
-		Variable = new(ParseState)ir_variable(Type, ralloc_asprintf(ParseState, "var_%s", Semantic), ir_var_out);
-	}
-
-	if (Variable)
-	{
-		// Up to this point, variables aren't contained in structs
-		*DestVariableType = Variable->type;
-		DeclInstructions->push_tail(Variable);
-		ParseState->symbols->add_variable(Variable);
-		Variable->centroid = OutputQualifier.Fields.bCentroid;
-		Variable->interpolation = OutputQualifier.Fields.InterpolationMode;
-		Variable->is_patch_constant = OutputQualifier.Fields.bIsPatchConstant;
-		ir_rvalue* VariableDeref = new(ParseState)ir_dereference_variable(Variable);
-		return VariableDeref;
-	}
-
-	if (Semantic && FCStringAnsi::Strnicmp(Semantic, "SV_", 3) == 0)
-	{
-		_mesa_glsl_warning(ParseState, "unrecognized system value output '%s'",
-			Semantic);
-	}
-
-	*DestVariableType = Type;
-
-	if (Frequency == HSF_HullShader && !OutputQualifier.Fields.bIsPatchConstant)
-	{
-		Type = glsl_type::get_array_instance(Type, ParseState->tessellation.outputcontrolpoints);
-	}
-
-	// This code-section replaces "layout(location=0) out struct { vec4 Data; } out_TEXCOORD0;" pattern to
-	// "layout(location=0) out vec4 out_TEXCOORD0;".
-
-	// Regular attribute
-	Variable = new(ParseState)ir_variable(
-		Type,
-		ralloc_asprintf(ParseState, "out_%s", Semantic),
-		ir_var_out
-		);
-	//Variable->read_only = true;
-	Variable->centroid = OutputQualifier.Fields.bCentroid;
-	Variable->interpolation = OutputQualifier.Fields.InterpolationMode;
-	Variable->is_patch_constant = OutputQualifier.Fields.bIsPatchConstant;
-
-	if (ParseState->bGenerateLayoutLocations)
-	{
-		ConfigureInOutVariableLayout(Frequency, ParseState, Semantic, Variable, ir_var_out);
-	}
-
-	DeclInstructions->push_tail(Variable);
-	ParseState->symbols->add_variable(Variable);
-
-	ir_dereference* VariableDeref = new(ParseState)ir_dereference_variable(Variable);
-
-	if (Frequency == HSF_HullShader && !OutputQualifier.Fields.bIsPatchConstant)
-	{
-		VariableDeref = new(ParseState)ir_dereference_array(VariableDeref, new(ParseState)ir_dereference_variable(ParseState->symbols->get_variable("gl_InvocationID")));
-	}
-
-	return VariableDeref;
-}
-
-/**
-* Generate an input semantic.
-* @param Frequency - The shader frequency.
-* @param ParseState - Parse state.
-* @param InputSemantic - The semantic name to generate.
-* @param InputQualifier - Qualifiers applied to the semantic.
-* @param InputVariableDeref - Deref for the argument variable.
-* @param DeclInstructions - IR to which declarations may be added.
-* @param PreCallInstructions - IR to which instructions may be added before the
-*                              entry point is called.
-*/
-static void GenShaderInputForVariable(
-	EHlslShaderFrequency Frequency,
-	_mesa_glsl_parse_state* ParseState,
-	const char* InputSemantic,
-	FSemanticQualifier InputQualifier,
-	ir_dereference* InputVariableDeref,
-	exec_list* DeclInstructions,
-	exec_list* PreCallInstructions,
-	int SemanticArraySize,
-	int SemanticArrayIndex
-	)
-{
-	const glsl_type* InputType = InputVariableDeref->type;
-
-	if (InputType->is_record())
-	{
-		for (uint32 i = 0; i < InputType->length; ++i)
-		{
-			const char* FieldSemantic = InputType->fields.structure[i].semantic;
-			const char* Semantic = 0;
-
-			if (InputSemantic && FieldSemantic)
-			{
-
-				_mesa_glsl_warning(ParseState, "semantic '%s' of field '%s' will be overridden by enclosing types' semantic '%s'",
-					InputType->fields.structure[i].semantic,
-					InputType->fields.structure[i].name,
-					InputSemantic);
-
-
-				FieldSemantic = 0;
-			}
-
-			if (InputSemantic && !FieldSemantic)
-			{
-				Semantic = ralloc_asprintf(ParseState, "%s%u", InputSemantic, i);
-				_mesa_glsl_warning(ParseState, "  creating semantic '%s' for struct field '%s'", Semantic, InputType->fields.structure[i].name);
-			}
-			else if (!InputSemantic && FieldSemantic)
-			{
-				Semantic = FieldSemantic;
-			}
-			else
-			{
-				Semantic = 0;
-			}
-
-			if (InputType->fields.structure[i].type->is_record() ||
-				Semantic)
-			{
-				FSemanticQualifier Qualifier = InputQualifier;
-				if (Qualifier.Packed == 0)
-				{
-					Qualifier.Fields.bCentroid = InputType->fields.structure[i].centroid;
-					Qualifier.Fields.InterpolationMode = InputType->fields.structure[i].interpolation;
-					Qualifier.Fields.bIsPatchConstant = InputType->fields.structure[i].patchconstant;
-				}
-
-				ir_dereference_record* FieldDeref = new(ParseState)ir_dereference_record(
-					InputVariableDeref->clone(ParseState, NULL),
-					InputType->fields.structure[i].name);
-				GenShaderInputForVariable(
-					Frequency,
-					ParseState,
-					Semantic,
-					Qualifier,
-					FieldDeref,
-					DeclInstructions,
-					PreCallInstructions,
-					SemanticArraySize,
-					SemanticArrayIndex
-					);
-			}
-			else
-			{
-				_mesa_glsl_error(
-					ParseState,
-					"field '%s' in input structure '%s' does not specify a semantic",
-					InputType->fields.structure[i].name,
-					InputType->name
-					);
-			}
-		}
-	}
-	else if (InputType->is_array() || InputType->is_inputpatch() || InputType->is_outputpatch())
-	{
-		int BaseIndex = 0;
-		const char* Semantic = 0;
-		check(InputSemantic);
-		ParseSemanticAndIndex(ParseState, InputSemantic, &Semantic, &BaseIndex);
-		check(BaseIndex >= 0);
-		check(InputType->is_array() || InputType->is_inputpatch() || InputType->is_outputpatch());
-		const unsigned ElementCount = InputType->is_array() ? InputType->length : InputType->patch_length;
-
-		{
-			//check(!InputQualifier.Fields.bIsPatchConstant);
-			InputQualifier.Fields.bIsPatchConstant = false;
-		}
-
-		for (unsigned i = 0; i < ElementCount; ++i)
-		{
-			ir_dereference_array* ArrayDeref = new(ParseState)ir_dereference_array(
-				InputVariableDeref->clone(ParseState, NULL),
-				new(ParseState)ir_constant((unsigned)i)
-				);
-			GenShaderInputForVariable(
-				Frequency,
-				ParseState,
-				ralloc_asprintf(ParseState, "%s%u", Semantic, BaseIndex + i),
-				InputQualifier,
-				ArrayDeref,
-				DeclInstructions,
-				PreCallInstructions,
-				SemanticArraySize,
-				SemanticArrayIndex
-				);
-		}
-	}
-	else
-	{
-		bool ApplyClipSpaceAdjustment = false;
-		ir_rvalue* SrcValue = GenShaderInputSemantic(Frequency, ParseState, InputSemantic,
-			InputQualifier, InputType, DeclInstructions, SemanticArraySize,
-			SemanticArrayIndex, ApplyClipSpaceAdjustment);
-
-		if (SrcValue)
-		{
-			YYLTYPE loc;
-
-			if (ParseState->adjust_clip_space_dx11_to_opengl && ApplyClipSpaceAdjustment)
-			{
-				// This is for input of gl_Position into geometry shader only.
-				check(Frequency == HSF_GeometryShader && InputSemantic && FCStringAnsi::Stricmp(InputSemantic, "SV_Position") == 0);
-
-				// Generate a local variable to do the conversion in, keeping source type.
-				ir_variable* TempVariable = new(ParseState)ir_variable(SrcValue->type, NULL, ir_var_temporary);
-				PreCallInstructions->push_tail(TempVariable);
-
-				// Assign input to this variable
-				ir_dereference_variable* TempVariableDeref = new(ParseState)ir_dereference_variable(TempVariable);
-				PreCallInstructions->push_tail(
-					new(ParseState)ir_assignment(
-					TempVariableDeref,
-					SrcValue
-					)
-					);
-
-				{
-					// TempVariable.y = -TempVariable.y;
-					PreCallInstructions->push_tail(
-						new(ParseState)ir_assignment(
-						new(ParseState)ir_swizzle(TempVariableDeref->clone(ParseState, NULL), 1, 0, 0, 0, 1),
-						new(ParseState)ir_expression(ir_unop_neg,
-						glsl_type::float_type,
-						new(ParseState)ir_swizzle(TempVariableDeref->clone(ParseState, NULL), 1, 0, 0, 0, 1),
-						NULL)
-						)
-						);
-				}
-
-				// Use TempVariable anywhere you would have otherwise used SrcValue going forward...
-				SrcValue = TempVariableDeref->clone(ParseState, NULL);
-			}
-
-			apply_type_conversion(InputType, SrcValue, PreCallInstructions, ParseState, true, &loc);
-			PreCallInstructions->push_tail(
-				new(ParseState)ir_assignment(
-				InputVariableDeref->clone(ParseState, NULL),
-				SrcValue
-				)
-				);
-		}
-	}
-}
-
-
-/**
-* Generate a shader input.
-* @param Frequency - The shader frequency.
-* @param ParseState - Parse state.
-* @param InputSemantic - The semantic name to generate.
-* @param InputQualifier - Qualifiers applied to the semantic.
-* @param InputType - Value type.
-* @param DeclInstructions - IR to which declarations may be added.
-* @param PreCallInstructions - IR to which instructions may be added before the
-*                              entry point is called.
-* @returns the IR variable deref for the semantic.
-*/
-static ir_dereference_variable* GenShaderInput(
-	EHlslShaderFrequency Frequency,
-	_mesa_glsl_parse_state* ParseState,
-	const char* InputSemantic,
-	FSemanticQualifier InputQualifier,
-	const glsl_type* InputType,
-	exec_list* DeclInstructions,
-	exec_list* PreCallInstructions)
-{
-	ir_variable* TempVariable = new(ParseState)ir_variable(
-		InputType,
-		NULL,
-		ir_var_temporary);
-	ir_dereference_variable* TempVariableDeref = new(ParseState)ir_dereference_variable(TempVariable);
-	PreCallInstructions->push_tail(TempVariable);
-
-	//check ( InputSemantic ?  (FCStringAnsi::Strnicmp(InputSemantic, "SV_", 3) ==0) : true);
-
-
-	// everything that's not an Outputpatch is patch constant. System values are treated specially
-	if (Frequency == HSF_DomainShader && !InputType->is_outputpatch())
-	{
-		InputQualifier.Fields.bIsPatchConstant = true;
-	}
-
-	if ((Frequency == HSF_GeometryShader && TempVariableDeref->type->is_array()) ||
-		(Frequency == HSF_HullShader && TempVariableDeref->type->is_inputpatch()) ||
-		(Frequency == HSF_DomainShader && TempVariableDeref->type->is_outputpatch())
-		)
-	{
-		check(InputType->is_array() || InputType->is_inputpatch() || InputType->is_outputpatch());
-		check(InputType->length || InputType->patch_length);
-
-
-		const unsigned ElementCount = InputType->is_array() ? InputType->length : InputType->patch_length;
-
-		if (!ParseState->symbols->get_variable("gl_in"))
-		{
-			// Create a built-in OpenGL variable gl_in[] containing built-in types.
-			// This variable will be used for OpenGL optimization by IR, so IR must know about it,
-			// but will not end up in final GLSL code.
-
-			// It has to be created here, as it contains multiple built-in variables in one interface block,
-			// which is not usual, so avoiding special cases in code.
-
-			glsl_struct_field *BuiltinFields = ralloc_array(ParseState, glsl_struct_field, 3);
-			memset(BuiltinFields, 0, 3 * sizeof(glsl_struct_field));
-
-			BuiltinFields[0].type = glsl_type::vec4_type;
-			BuiltinFields[0].name = ralloc_strdup(ParseState, "gl_Position");
-			BuiltinFields[1].type = glsl_type::float_type;
-			BuiltinFields[1].name = ralloc_strdup(ParseState, "gl_PointSize");
-			BuiltinFields[2].type = glsl_type::get_array_instance(glsl_type::float_type, 6);	// magic number is gl_MaxClipDistances
-			BuiltinFields[2].name = ralloc_strdup(ParseState, "gl_ClipDistance");
-
-			const glsl_type* BuiltinStruct = glsl_type::get_record_instance(BuiltinFields, 3, "gl_PerVertex");
-			const glsl_type* BuiltinArray = glsl_type::get_array_instance(BuiltinStruct, ElementCount);
-			ir_variable* BuiltinVariable = new(ParseState)ir_variable(BuiltinArray, "gl_in", ir_var_in);
-			BuiltinVariable->read_only = true;
-			BuiltinVariable->is_interface_block = true;
-			DeclInstructions->push_tail(BuiltinVariable);
-			ParseState->symbols->add_variable(BuiltinVariable);
-		}
-
-		for (unsigned i = 0; i < ElementCount; ++i)
-		{
-			ir_dereference_array* ArrayDeref = new(ParseState)ir_dereference_array(
-				TempVariableDeref->clone(ParseState, NULL),
-				new(ParseState)ir_constant((unsigned)i)
-				);
-			// Parse input variable
-			GenShaderInputForVariable(
-				Frequency,
-				ParseState,
-				InputSemantic,
-				InputQualifier,
-				ArrayDeref,
-				DeclInstructions,
-				PreCallInstructions,
-				ElementCount,
-				i
-				);
-		}
-	}
-	else
-	{
-		GenShaderInputForVariable(
-			Frequency,
-			ParseState,
-			InputSemantic,
-			InputQualifier,
-			TempVariableDeref,
-			DeclInstructions,
-			PreCallInstructions,
-			0,
-			0
-			);
-	}
-	return TempVariableDeref;
-}
-
-/**
-* Generate an output semantic.
-* @param Frequency - The shader frequency.
-* @param ParseState - Parse state.
-* @param OutputSemantic - The semantic name to generate.
-* @param OutputQualifier - Qualifiers applied to the semantic.
-* @param OutputVariableDeref - Deref for the argument variable.
-* @param DeclInstructions - IR to which declarations may be added.
-* @param PostCallInstructions - IR to which instructions may be added after the
-*                               entry point returns.
-*/
-static void GenShaderOutputForVariable(
-	EHlslShaderFrequency Frequency,
-	_mesa_glsl_parse_state* ParseState,
-	const char* OutputSemantic,
-	FSemanticQualifier OutputQualifier,
-	ir_dereference* OutputVariableDeref,
-	exec_list* DeclInstructions,
-	exec_list* PostCallInstructions
-	)
-{
-	const glsl_type* OutputType = OutputVariableDeref->type;
-	if (OutputType->is_record())
-	{
-		for (uint32 i = 0; i < OutputType->length; ++i)
-		{
-			const char* FieldSemantic = OutputType->fields.structure[i].semantic;
-			const char* Semantic = 0;
-
-			if (OutputSemantic && FieldSemantic)
-			{
-
-				_mesa_glsl_warning(ParseState, "semantic '%s' of field '%s' will be overridden by enclosing types' semantic '%s'",
-					OutputType->fields.structure[i].semantic,
-					OutputType->fields.structure[i].name,
-					OutputSemantic);
-
-
-				FieldSemantic = 0;
-			}
-
-			if (OutputSemantic && !FieldSemantic)
-			{
-				Semantic = ralloc_asprintf(ParseState, "%s%u", OutputSemantic, i);
-				_mesa_glsl_warning(ParseState, "  creating semantic '%s' for struct field '%s'", Semantic, OutputType->fields.structure[i].name);
-			}
-			else if (!OutputSemantic && FieldSemantic)
-			{
-				Semantic = FieldSemantic;
-			}
-			else
-			{
-				Semantic = 0;
-			}
-
-			if (OutputType->fields.structure[i].type->is_record() ||
-				Semantic
-				)
-			{
-				FSemanticQualifier Qualifier = OutputQualifier;
-				if (Qualifier.Packed == 0)
-				{
-					Qualifier.Fields.bCentroid = OutputType->fields.structure[i].centroid;
-					Qualifier.Fields.InterpolationMode = OutputType->fields.structure[i].interpolation;
-					Qualifier.Fields.bIsPatchConstant = OutputType->fields.structure[i].patchconstant;
-				}
-
-				// Dereference the field and generate shader outputs for the field.
-				ir_dereference* FieldDeref = new(ParseState)ir_dereference_record(
-					OutputVariableDeref->clone(ParseState, NULL),
-					OutputType->fields.structure[i].name);
-				GenShaderOutputForVariable(
-					Frequency,
-					ParseState,
-					Semantic,
-					Qualifier,
-					FieldDeref,
-					DeclInstructions,
-					PostCallInstructions
-					);
-			}
-			else
-			{
-				_mesa_glsl_error(
-					ParseState,
-					"field '%s' in output structure '%s' does not specify a semantic",
-					OutputType->fields.structure[i].name,
-					OutputType->name
-					);
-			}
-		}
-	}
-	// TODO clean this up!!
-	else if ((OutputType->is_array() || OutputType->is_outputpatch()))
-	{
-		if (OutputSemantic)
-		{
-			int BaseIndex = 0;
-			const char* Semantic = 0;
-
-			ParseSemanticAndIndex(ParseState, OutputSemantic, &Semantic, &BaseIndex);
-
-			const unsigned ElementCount = OutputType->is_array() ? OutputType->length : (OutputType->patch_length);
-
-			for (unsigned i = 0; i < ElementCount; ++i)
-			{
-				ir_dereference_array* ArrayDeref = new(ParseState)ir_dereference_array(
-					OutputVariableDeref->clone(ParseState, NULL),
-					new(ParseState)ir_constant((unsigned)i)
-					);
-				GenShaderOutputForVariable(
-					Frequency,
-					ParseState,
-					ralloc_asprintf(ParseState, "%s%u", Semantic, BaseIndex + i),
-					OutputQualifier,
-					ArrayDeref,
-					DeclInstructions,
-					PostCallInstructions
-					);
-			}
-		}
-		else
-		{
-			_mesa_glsl_error(ParseState, "entry point does not specify a semantic for its return value");
-		}
-	}
-	else
-	{
-		if (OutputSemantic)
-		{
-			YYLTYPE loc;
-			ir_rvalue* Src = OutputVariableDeref->clone(ParseState, NULL);
-			const glsl_type* DestVariableType = NULL;
-			bool ApplyClipSpaceAdjustment = false;
-			bool ApplyClampPowerOfTwo = false;
-			ir_rvalue* DestVariableDeref = GenShaderOutputSemantic(Frequency, ParseState, OutputSemantic,
-				OutputQualifier, OutputType, DeclInstructions, &DestVariableType, ApplyClipSpaceAdjustment, ApplyClampPowerOfTwo);
-
-			apply_type_conversion(DestVariableType, Src, PostCallInstructions, ParseState, true, &loc);
-
-			if (ParseState->adjust_clip_space_dx11_to_opengl && ApplyClipSpaceAdjustment)
-			{
-				// Src.y = -Src.y;
-				PostCallInstructions->push_tail(
-					new(ParseState)ir_assignment(
-					new(ParseState)ir_swizzle(Src->clone(ParseState, NULL), 1, 0, 0, 0, 1),
-					new(ParseState)ir_expression(ir_unop_neg,
-					glsl_type::float_type,
-					new(ParseState)ir_swizzle(Src->clone(ParseState, NULL), 1, 0, 0, 0, 1),
-					NULL)
-					)
-					);
-			}
-
-			// GLSL doesn't support pow2 partitioning, so we treat pow2 as integer partitioning and
-			// manually compute the next power of two via exp2(pow(ceil(log2(Src)));
-			if (ApplyClampPowerOfTwo)
-			{
-				ir_variable* temp = new(ParseState)ir_variable(glsl_type::float_type, NULL, ir_var_temporary);
-
-				PostCallInstructions->push_tail(temp);
-
-				PostCallInstructions->push_tail(
-					new(ParseState)ir_assignment(
-					new(ParseState)ir_dereference_variable(temp),
-					new(ParseState)ir_expression(ir_unop_exp2,
-					new(ParseState)ir_expression(ir_unop_ceil,
-					new(ParseState)ir_expression(ir_unop_log2,
-					glsl_type::float_type,
-					Src->clone(ParseState, NULL),
-					NULL
-					)
-					)
-					)
-					)
-					);
-
-				// assign pow2 clamped variable to output variable
-				PostCallInstructions->push_tail(
-					new(ParseState)ir_assignment(
-					DestVariableDeref->clone(ParseState, NULL),
-					new(ParseState)ir_dereference_variable(temp)
-					)
-					);
-			}
-			else
-			{
-				PostCallInstructions->push_tail(new(ParseState)ir_assignment(DestVariableDeref, Src));
-			}
-		}
-		else
-		{
-			_mesa_glsl_error(ParseState, "entry point does not specify a semantic for its return value");
-		}
-	}
-}
-/**
-* Generate an output semantic.
-* @param Frequency - The shader frequency.
-* @param ParseState - Parse state.
-* @param OutputSemantic - The semantic name to generate.
-* @param OutputQualifier - Qualifiers applied to the semantic.
-* @param OutputType - Value type.
-* @param DeclInstructions - IR to which declarations may be added.
-* @param PreCallInstructions - IR to which instructions may be added before the
-entry point is called.
-* @param PostCallInstructions - IR to which instructions may be added after the
-*                               entry point returns.
-* @returns the IR variable deref for the semantic.
-*/
-static ir_dereference_variable* GenShaderOutput(
-	EHlslShaderFrequency Frequency,
-	_mesa_glsl_parse_state* ParseState,
-	const char* OutputSemantic,
-	FSemanticQualifier OutputQualifier,
-	const glsl_type* OutputType,
-	exec_list* DeclInstructions,
-	exec_list* PreCallInstructions,
-	exec_list* PostCallInstructions
-	)
-{
-	// Generate a local variable to hold the output.
-	ir_variable* TempVariable = new(ParseState)ir_variable(
-		OutputType,
-		NULL,
-		ir_var_temporary);
-	ir_dereference_variable* TempVariableDeref = new(ParseState)ir_dereference_variable(TempVariable);
-	PreCallInstructions->push_tail(TempVariable);
-	GenShaderOutputForVariable(
-		Frequency,
-		ParseState,
-		OutputSemantic,
-		OutputQualifier,
-		TempVariableDeref,
-		DeclInstructions,
-		PostCallInstructions);
-	return TempVariableDeref;
-}
-
-static void GenerateAppendFunctionBody(
-	_mesa_glsl_parse_state* ParseState,
-	exec_list* DeclInstructions,
-	const glsl_type* geometry_append_type
-	)
-{
-	ir_function *func = ParseState->symbols->get_function("OutputStream_Append");
-	check(func);
-
-	exec_list comparison_parameter;
-	ir_variable* var = new(ParseState)ir_variable(geometry_append_type, ralloc_asprintf(ParseState, "arg0"), ir_var_in);
-	comparison_parameter.push_tail(var);
-
-	bool is_exact = false;
-	ir_function_signature *sig = func->matching_signature(&comparison_parameter, &is_exact);
-	check(sig && is_exact);
-	var = (ir_variable*)sig->parameters.get_head();
-
-	//	{
-	//		const glsl_type* output_type = var->type;
-	//		_mesa_glsl_warning(ParseState, "GenerateAppendFunctionBody: parsing argument struct '%s'", output_type->name );
-	//		int indexof_RenderTargetArrayIndex = -1;
-	//		for (int i = 0; i < output_type->length; i++)
-	//		{
-	//			_mesa_glsl_warning(ParseState, "   name '%s' : semantic '%s'", output_type->fields.structure[i].name, output_type->fields.structure[i].semantic );
-	//		}
-	//	}
-
-	// Generate assignment instructions from function argument to out variables
-	FSemanticQualifier OutputQualifier;
-	ir_dereference_variable* TempVariableDeref = new(ParseState)ir_dereference_variable(var);
-	GenShaderOutputForVariable(
-		HSF_GeometryShader,
-		ParseState,
-		NULL,
-		OutputQualifier,
-		TempVariableDeref,
-		DeclInstructions,
-		&sig->body);
-
-	// If the output structure type contains a SV_RenderTargetArrayIndex semantic, add a custom user output semantic.
-	// It's used to pass layer index to pixel shader, as GLSL 1.50 doesn't allow pixel shader to read from gl_Layer.
-	const glsl_type* output_type = var->type;
-	int indexof_RenderTargetArrayIndex = -1;
-/*
-	for (uint32 i = 0; i < output_type->length; i++)
-	{
-		if (output_type->fields.structure[i].semantic && (FCStringAnsi::Strcmp(output_type->fields.structure[i].semantic, "SV_RenderTargetArrayIndex") == 0))
-		{
-			indexof_RenderTargetArrayIndex = i;
-			break;
-		}
-	}
-
-	if (indexof_RenderTargetArrayIndex != -1)
-	{
-		// Add the new member with semantic
-		glsl_struct_field field;
-		field.type = output_type->fields.structure[indexof_RenderTargetArrayIndex].type;
-		field.name = "HLSLCCLayerIndex";
-		field.semantic = CUSTOM_LAYER_INDEX_SEMANTIC;
-		field.centroid = 0;
-		field.interpolation = ir_interp_qualifier_flat;
-		field.geometryinput = 0;
-		field.patchconstant = 0;
-
-		glsl_type* non_const_type = (glsl_type*)output_type;
-		non_const_type->add_structure_member(&field);
-
-		// Create new out variable for the new member and generate assignment that will copy input's layer index field to it
-		FSemanticQualifier Qualifier;
-		Qualifier.Fields.bCentroid = 0;
-		Qualifier.Fields.InterpolationMode = ir_interp_qualifier_flat;
-
-		const glsl_type* new_output_type = ((ir_variable*)sig->parameters.get_head())->type;
-		GenShaderOutputForVariable(
-			HSF_GeometryShader,
-			ParseState,
-			CUSTOM_LAYER_INDEX_SEMANTIC,
-			Qualifier,
-			new(ParseState)ir_dereference_record(var, new_output_type->fields.structure[indexof_RenderTargetArrayIndex].name),
-			DeclInstructions,
-			&sig->body,
-			0,
-			0
-			);
-	}
-*/
-	// Call EmitVertex()
-	ir_function *emitVertexFunc = ParseState->symbols->get_function("EmitVertex");
-	check(emitVertexFunc);
-	check(emitVertexFunc->signatures.get_head() == emitVertexFunc->signatures.get_tail());
-	ir_function_signature *emitVertexSig = (ir_function_signature *)emitVertexFunc->signatures.get_head();
-	exec_list actual_parameter;
-	sig->body.push_tail(new(ParseState)ir_call(emitVertexSig, NULL, &actual_parameter));
-}
-
-bool FVulkanCodeBackend::GenerateMain(
-	EHlslShaderFrequency Frequency,
-	const char* EntryPoint,
-	exec_list* Instructions,
-	_mesa_glsl_parse_state* ParseState)
-{
-	// Don't force coordinate system adjustment from GLSL->Vulkan as we transition to flipping the viewport instead of gl_Position.y coordinate
-	ParseState->adjust_clip_space_dx11_to_opengl = false;
-
-	{
-		// Set up origin_upper_left for gl_FragCoord, depending on HLSLCC_DX11ClipSpace flag presence.
-		FSystemValue* SystemValues = SystemValueTable[HSF_PixelShader];
-		for (int i = 0; SystemValues[i].Semantic != NULL; ++i)
-		{
-			if (FCStringAnsi::Stricmp(SystemValues[i].GlslName, "gl_FragCoord") == 0)
-			{
-				// Always disable layout(origin_upper_left) attribute as we transition to flipping the viewport instead of gl_Position.y coordinate
-				SystemValues[i].bOriginUpperLeft = false;// !ParseState->adjust_clip_space_dx11_to_opengl;// false;
-				break;
-			}
-		}
-	}
-
-	ir_function_signature* EntryPointSig = FindEntryPointFunction(Instructions, ParseState, EntryPoint);
-	if (EntryPointSig)
-	{
-		void* TempMemContext = ralloc_context(NULL);
-		exec_list DeclInstructions;
-		exec_list PreCallInstructions;
-		exec_list ArgInstructions;
-		exec_list PostCallInstructions;
-		const glsl_type* geometry_append_type = NULL;
-
-		ParseState->maxvertexcount = EntryPointSig->maxvertexcount;
-
-		ParseState->tessellation = EntryPointSig->tessellation;
-
-		ParseState->symbols->push_scope();
-
-		foreach_iter(exec_list_iterator, Iter, EntryPointSig->parameters)
-		{
-			ir_variable *const Variable = (ir_variable *)Iter.get();
-			if (Variable->semantic != NULL || Variable->type->is_record()
-				|| (Frequency == HSF_GeometryShader && (Variable->type->is_outputstream() || Variable->type->is_array()))
-				|| (Frequency == HSF_HullShader && (Variable->type->is_patch()))
-				|| (Frequency == HSF_DomainShader && (Variable->type->is_outputpatch()))
-				)
-			{
-				FSemanticQualifier Qualifier;
-				Qualifier.Fields.bCentroid = Variable->centroid;
-				Variable->centroid = 0;
-				Qualifier.Fields.InterpolationMode = Variable->interpolation;
-				Variable->interpolation = 0;
-				Qualifier.Fields.bIsPatchConstant = Variable->is_patch_constant;
-				Variable->is_patch_constant = 0;
-
-				ir_dereference_variable* ArgVarDeref = NULL;
-				switch (Variable->mode)
-				{
-				case ir_var_in:
-					if (Frequency == HSF_GeometryShader && Variable->type->is_array())
-					{
-						// Remember information about geometry input type globally
-						ParseState->geometryinput = Variable->geometryinput;
-					}
-
-					if (Frequency == HSF_PixelShader)
-					{
-						// Replace SV_RenderTargetArrayIndex in
-						// input structure semantic with custom semantic.
-						if (Variable->semantic && (FCStringAnsi::Strcmp(Variable->semantic, "SV_RenderTargetArrayIndex") == 0))
-						{
-							//							_mesa_glsl_warning(ParseState, "Replacing semantic of variable '%s' with our custom one", Variable->name);
-							Variable->semantic = ralloc_strdup(Variable, CUSTOM_LAYER_INDEX_SEMANTIC);
-							Variable->interpolation = ir_interp_qualifier_flat;
-						}
-						else if (Variable->type->is_record())
-						{
-							const glsl_type* output_type = Variable->type;
-							int indexof_RenderTargetArrayIndex = -1;
-/*
-							for (uint32 i = 0; i < output_type->length; i++)
-							{
-								if (Variable->type->fields.structure[i].semantic && (FCStringAnsi::Strcmp(Variable->type->fields.structure[i].semantic, "SV_RenderTargetArrayIndex") == 0))
-								{
-									indexof_RenderTargetArrayIndex = i;
-									break;
-								}
-							}
-
-							if (indexof_RenderTargetArrayIndex != -1)
-							{
-								// _mesa_glsl_warning(ParseState, "Replacing semantic of member %d of variable '%s' with our custom one", indexof_RenderTargetArrayIndex, Variable->name);
-								// Replace the member with one with semantic
-								glsl_struct_field field;
-								field.type = Variable->type->fields.structure[indexof_RenderTargetArrayIndex].type;
-								field.name = Variable->type->fields.structure[indexof_RenderTargetArrayIndex].name;
-								field.semantic = CUSTOM_LAYER_INDEX_SEMANTIC;
-								field.centroid = 0;
-								field.interpolation = ir_interp_qualifier_flat;
-								field.geometryinput = 0;
-								field.patchconstant = 0;
-
-								glsl_type* non_const_type = (glsl_type*)output_type;
-								non_const_type->replace_structure_member(indexof_RenderTargetArrayIndex, &field);
-							}
-*/
-						}
-					}
-
-					ArgVarDeref = GenShaderInput(
-						Frequency,
-						ParseState,
-						Variable->semantic,
-						Qualifier,
-						Variable->type,
-						&DeclInstructions,
-						&PreCallInstructions
-						);
-					break;
-				case ir_var_out:
-					if (Frequency == HSF_PixelShader && Variable->semantic && (strcmp(Variable->semantic, "SV_Depth") == 0))
-					{
-						bExplicitDepthWrites = true;
-					}
-					
-					ArgVarDeref = GenShaderOutput(
-						Frequency,
-						ParseState,
-						Variable->semantic,
-						Qualifier,
-						Variable->type,
-						&DeclInstructions,
-						&PreCallInstructions,
-						&PostCallInstructions
-						);
-					break;
-				case ir_var_inout:
-				{
-					check(Frequency == HSF_GeometryShader);
-					// This is an output stream for geometry shader. It's not referenced as a variable inside the function,
-					// instead OutputStream.Append(vertex) and OutputStream.RestartStrip() are called, and this variable
-					// has already been optimized out of them in ast_to_hir translation.
-
-					// Generate a local variable to add to arguments. It won't be referenced anywhere, so it should get optimized out.
-					ir_variable* TempVariable = new(ParseState)ir_variable(
-						Variable->type,
-						NULL,
-						ir_var_temporary);
-					ArgVarDeref = new(ParseState)ir_dereference_variable(TempVariable);
-					PreCallInstructions.push_tail(TempVariable);
-
-					// We need to move this information somewhere safer, as this pseudo-variable will get optimized out of existence
-					ParseState->outputstream_type = Variable->type->outputstream_type;
-
-					check(Variable->type->is_outputstream());
-					check(Variable->type->inner_type->is_record());
-
-					geometry_append_type = Variable->type->inner_type;
-				}
-				break;
-				default:
-				{
-					_mesa_glsl_error(
-						ParseState,
-						"entry point parameter '%s' must be an input or output",
-						Variable->name
-						);
-				}
-				}
-				ArgInstructions.push_tail(ArgVarDeref);
-			}
-			else
-			{
-				_mesa_glsl_error(ParseState, "entry point parameter "
-					"'%s' does not specify a semantic", Variable->name);
-			}
-		}
-
-		// The function's return value should have an output semantic if it's not void.
-		ir_dereference_variable* EntryPointReturn = NULL;
-		if (EntryPointSig->return_type->is_void() == false)
-		{
-			FSemanticQualifier Qualifier;
-			EntryPointReturn = GenShaderOutput(
-				Frequency,
-				ParseState,
-				EntryPointSig->return_semantic,
-				Qualifier,
-				EntryPointSig->return_type,
-				&DeclInstructions,
-				&PreCallInstructions,
-				&PostCallInstructions
-				);
-		}
-
-		if (Frequency == HSF_GeometryShader)
-		{
-			GenerateAppendFunctionBody(
-				ParseState,
-				&DeclInstructions,
-				geometry_append_type
-				);
-		}
-
-		/*
-		we map the HLSL hull shader to this GLSL main function
-		for the most parts, we treat variables of InputPatch and OutputPatch as arrays of the inner type
-
-		build input patch from shader input interface blocks
-		call hull shader main function with input patch and current control point id (gl_InvocationID)
-		copy hull shader main result for the current control point to the proper shader output interface block element
-		barrier
-		(so all instances have computed the per control point data)
-		build patch constant function input (of type output patch) from the shader output interface blocks
-		(need to do this, since this is the only shader variable shared between control points running in parallel)
-
-		if control point id (gl_InvocationID) is 0
-		call patch constant function with the output patch as an input
-		copy the patch constant result to the "patch" shader output interface block
-		*/
-
-		if (Frequency == HSF_HullShader)
-		{
-
-			ir_function_signature* PatchConstantSig = FindPatchConstantFunction(Instructions, ParseState);
-
-			if (!PatchConstantSig)
-			{
-				_mesa_glsl_error(ParseState, "patch constant function `%s' not found", ParseState->tessellation.patchconstantfunc);
-			}
-
-			const glsl_type* OutputPatchType = glsl_type::get_templated_instance(EntryPointReturn->type, "OutputPatch", 0, ParseState->tessellation.outputcontrolpoints);
-
-			ir_variable* OutputPatchVar = new(ParseState)ir_variable(OutputPatchType, NULL, ir_var_temporary);
-
-
-			// call barrier() to ensure that all threads have computed the per-patch computation
-			{
-				// We can't just use the symbol table b/c it only has the HLSL and not the GLSL barrier functions in it
-				foreach_iter(exec_list_iterator, Iter, *Instructions)
-				{
-					ir_instruction *ir = (ir_instruction *)Iter.get();
-					ir_function *Function = ir->as_function();
-					if (Function && FCStringAnsi::Strcmp(Function->name, "barrier") == 0)
-					{
-						check(Function->signatures.get_head() == Function->signatures.get_tail());
-						exec_list VoidParameter;
-						ir_function_signature * BarrierFunctionSig = Function->matching_signature(&VoidParameter);
-						PostCallInstructions.push_tail(new(ParseState)ir_call(BarrierFunctionSig, NULL, &VoidParameter));
-					}
-				}
-			}
-
-			// reassemble output patch variable(for the patch constant function) from the shader outputs
-			GenShaderPatchConstantFunctionInputs(ParseState, OutputPatchVar, PostCallInstructions);
-
-			// call the entry point
-			if (PatchConstantSig)
-			{
-				CallPatchConstantFunction(ParseState, OutputPatchVar, PatchConstantSig, DeclInstructions, PostCallInstructions);
-			}
-		}
-
-		ParseState->symbols->pop_scope();
-
-		// Build the void main() function for GLSL.
-		ir_function_signature* MainSig = new(ParseState)ir_function_signature(glsl_type::void_type);
-		MainSig->is_defined = true;
-		MainSig->is_main = true;
-		MainSig->body.append_list(&PreCallInstructions);
-		MainSig->body.push_tail(new(ParseState)ir_call(EntryPointSig, EntryPointReturn, &ArgInstructions));
-		MainSig->body.append_list(&PostCallInstructions);
-		MainSig->maxvertexcount = EntryPointSig->maxvertexcount;
-		MainSig->is_early_depth_stencil = (EntryPointSig->is_early_depth_stencil && !bExplicitDepthWrites);
-		MainSig->wg_size_x = EntryPointSig->wg_size_x;
-		MainSig->wg_size_y = EntryPointSig->wg_size_y;
-		MainSig->wg_size_z = EntryPointSig->wg_size_z;
-		MainSig->tessellation = EntryPointSig->tessellation;
-
-		if (MainSig->is_early_depth_stencil && Frequency != HSF_PixelShader)
-		{
-			_mesa_glsl_error(ParseState, "'earlydepthstencil' attribute only applies to pixel shaders");
-		}
-
-		if (MainSig->maxvertexcount > 0 && Frequency != HSF_GeometryShader)
-		{
-			_mesa_glsl_error(ParseState, "'maxvertexcount' attribute only applies to geometry shaders");
-		}
-
-		if (MainSig->is_early_depth_stencil && ParseState->language_version < 310)
-		{
-			_mesa_glsl_error(ParseState, "'earlydepthstencil' attribute only supported on GLSL 4.30 target and later");
-		}
-
-		if (MainSig->wg_size_x > 0 && Frequency != HSF_ComputeShader)
-		{
-			_mesa_glsl_error(ParseState, "'num_threads' attribute only applies to compute shaders");
-		}
-
-		// in GLSL, unlike in HLSL fixed-function tessellator properties are specified on the domain shader
-		// and not the hull shader, so we specify them for both in the .usf shaders and then print a warning,
-		// similar to what fxc is doing
-
-		if (MainSig->tessellation.domain != GLSL_DOMAIN_NONE && (Frequency != HSF_HullShader && Frequency != HSF_DomainShader))
-		{
-			_mesa_glsl_warning(ParseState, "'domain' attribute only applies to hull or domain shaders");
-		}
-
-		if (MainSig->tessellation.outputtopology != GLSL_OUTPUTTOPOLOGY_NONE && Frequency != HSF_HullShader)
-		{
-			_mesa_glsl_warning(ParseState, "'outputtopology' attribute only applies to hull shaders");
-		}
-
-		if (MainSig->tessellation.partitioning != GLSL_PARTITIONING_NONE && Frequency != HSF_HullShader)
-		{
-			_mesa_glsl_warning(ParseState, "'partitioning' attribute only applies to hull shaders");
-		}
-
-		if (MainSig->tessellation.outputcontrolpoints > 0 && Frequency != HSF_HullShader)
-		{
-			_mesa_glsl_warning(ParseState, "'outputcontrolpoints' attribute only applies to hull shaders");
-		}
-
-		if (MainSig->tessellation.maxtessfactor > 0.0f && Frequency != HSF_HullShader)
-		{
-			_mesa_glsl_warning(ParseState, "'maxtessfactor' attribute only applies to hull shaders");
-		}
-
-		if (MainSig->tessellation.patchconstantfunc != 0 && Frequency != HSF_HullShader)
-		{
-			_mesa_glsl_warning(ParseState, "'patchconstantfunc' attribute only applies to hull shaders");
-		}
-
-		// Values that will be patched in later from the SPIRV
-		ir_function* MainFunction = new(ParseState)ir_function("main_00000000_00000000");
-		MainFunction->add_signature(MainSig);
-
-		Instructions->append_list(&DeclInstructions);
-		Instructions->push_tail(MainFunction);
-
-		ralloc_free(TempMemContext);
-
-		// Now that we have a proper main(), move global setup to main().
-		MoveGlobalInstructionsToMain(Instructions);
-	}
-	else
-	{
-		_mesa_glsl_error(ParseState, "shader entry point '%s' not "
-			"found", EntryPoint);
-	}
-
-	return true;
-}
-ir_function_signature*  FVulkanCodeBackend::FindPatchConstantFunction(exec_list* Instructions, _mesa_glsl_parse_state* ParseState)
-{
-	ir_function_signature* PatchConstantSig = 0;
-
-	// TODO refactor this and the fetching of the main siganture
-	foreach_iter(exec_list_iterator, Iter, *Instructions)
-	{
-		ir_instruction *ir = (ir_instruction *)Iter.get();
-		ir_function *Function = ir->as_function();
-		if (Function && FCStringAnsi::Strcmp(Function->name, ParseState->tessellation.patchconstantfunc) == 0)
-		{
-			int NumSigs = 0;
-			foreach_iter(exec_list_iterator, SigIter, *Function)
-			{
-				if (++NumSigs == 1)
-				{
-					PatchConstantSig = (ir_function_signature *)SigIter.get();
-				}
-			}
-			if (NumSigs == 1)
-			{
-				break;
-			}
-			else
-			{
-				_mesa_glsl_error(ParseState, "patch constant function "
-					"`%s' has multiple signatures", ParseState->tessellation.patchconstantfunc);
-			}
-		}
-	}
-
-	return PatchConstantSig;
-}
-
-void FVulkanCodeBackend::CallPatchConstantFunction(_mesa_glsl_parse_state* ParseState, ir_variable* OutputPatchVar, ir_function_signature* PatchConstantSig, exec_list& DeclInstructions, exec_list &PostCallInstructions)
+const char* VULKAN_SUBPASS_FETCH[9] = 
 {
-	exec_list PatchConstantArgs;
-	if (OutputPatchVar && !PatchConstantSig->parameters.is_empty())
-	{
-		PatchConstantArgs.push_tail(new(ParseState)ir_dereference_variable(OutputPatchVar));
-	}
-
-	ir_if* thread_if = new(ParseState)ir_if(
-		new(ParseState)ir_expression(
-		ir_binop_equal,
-		new (ParseState)ir_constant(
-		0
-		),
-		new (ParseState)ir_dereference_variable(
-		ParseState->symbols->get_variable("gl_InvocationID")
-		)
-		)
-		);
-
-	exec_list PrePatchConstCallInstructions;
-	exec_list PostPatchConstCallInstructions;
-
-	FSemanticQualifier Qualifier;
-	Qualifier.Fields.bIsPatchConstant = 1;
-
-	ir_dereference_variable* PatchConstantReturn = GenShaderOutput(
-		HSF_HullShader,
-		ParseState,
-		PatchConstantSig->return_semantic,
-		Qualifier,
-		PatchConstantSig->return_type,
-		&DeclInstructions,
-		&PrePatchConstCallInstructions,
-		&PostPatchConstCallInstructions
-		);
-
-	thread_if->then_instructions.append_list(&PrePatchConstCallInstructions);
-	thread_if->then_instructions.push_tail(new(ParseState)ir_call(PatchConstantSig, PatchConstantReturn, &PatchConstantArgs));
-	thread_if->then_instructions.append_list(&PostPatchConstCallInstructions);
-
-	PostCallInstructions.push_tail(thread_if);
-}
-
-
-/*
-reassemble output patch variable (for the patch constant function) from the shader outputs
-turn this: (from the GenOutputs of calling the entry point main)
-
-out_InnerMember[gl_InvocationID].Data = t2.Middle.Inner.Value;
-
-into this:
-
-//output_patch<FPNTessellationHSToDS> FPNTessellationHSToDS t3[3]; //output_patch<FPNTessellationHSToDS> ;
-t3[0].Middle.Inner.Value = out_InnerMember[0].Data;
-t3[1].Middle.Inner.Value = out_InnerMember[1].Data;
-t3[2].Middle.Inner.Value = out_InnerMember[2].Data;
-
-*/
-
-void FVulkanCodeBackend::GenShaderPatchConstantFunctionInputs(_mesa_glsl_parse_state* ParseState, ir_variable* OutputPatchVar, exec_list &PostCallInstructions)
+	"VulkanSubpassDepthFetch",
+	"VulkanSubpassFetch0",
+	"VulkanSubpassFetch1",
+	"VulkanSubpassFetch2",
+	"VulkanSubpassFetch3",
+	"VulkanSubpassFetch4",
+	"VulkanSubpassFetch5",
+	"VulkanSubpassFetch6",
+	"VulkanSubpassFetch7"
+};
+const char* VULKAN_SUBPASS_FETCH_VAR[9] = 
 {
-	PostCallInstructions.push_tail(OutputPatchVar);
-	foreach_iter(exec_list_iterator, Iter, PostCallInstructions)
-	{
-		ir_instruction *ir = (ir_instruction *)Iter.get();
-
-		ir_assignment* assignment = ir->as_assignment();
-
-		if (!assignment)
-		{
-			continue;
-		}
-
-		ir_dereference_record* lhs = assignment->lhs->as_dereference_record();
-		ir_rvalue* rhs = assignment->rhs;
-
-		if (!rhs)
-		{
-			continue;
-		}
-
-		// Check whether LHS is wrapped into an array.
-		// This might be the case on the OpenGL backend, but not necessarily on the Vulkan backend.
-		ir_dereference_array* lhs_array = nullptr;
-		if (lhs)
-		{
-			lhs_array = lhs->record->as_dereference_array();
-		}
-		else
-		{
-			lhs_array = assignment->lhs->as_dereference_array();
-		}
-
-		if (!lhs_array)
-		{
-			continue;
-		}
-
-		ir_dereference_variable* OutputPatchArrayIndex = lhs_array->array_index->as_dereference_variable();
-		ir_dereference_variable* OutputPatchArray = lhs_array->array->as_dereference_variable();
-
-		if (!OutputPatchArrayIndex)
-		{
-			continue;
-		}
-
-		if (0 != FCStringAnsi::Strcmp(OutputPatchArrayIndex->var->name, "gl_InvocationID"))
-		{
-			continue;
-		}
-
-		if (!OutputPatchArray)
-		{
-			continue;
-		}
-
-		for (int OutputVertex = 0; OutputVertex < ParseState->tessellation.outputcontrolpoints; ++OutputVertex)
-		{
-			struct Helper
-			{
-				// the struct inside the output patch can have the actual outputs with semantics nested inside,
-
-				static void ReplaceVariableDerefWithArrayDeref(ir_instruction* Node, ir_dereference_array* ArrayDereference)
-				{
-					if (ir_dereference_record* AsRecord = Node->as_dereference_record())
-					{
-						if (AsRecord->record->as_dereference_variable())
-						{
-							AsRecord->record = ArrayDereference;
-						}
-						else
-						{
-							ReplaceVariableDerefWithArrayDeref(AsRecord->record, ArrayDereference);
-						}
-					}
-					else if (ir_dereference_array* AsArray = Node->as_dereference_array())
-					{
-						if (AsArray->array->as_dereference_variable())
-						{
-							AsArray->array = ArrayDereference;
-						}
-						else
-						{
-							ReplaceVariableDerefWithArrayDeref(AsArray->array, ArrayDereference);
-						}
-					}
-					else
-					{
-						check(false);
-					}
-				}
-			};
-
-			ir_dereference_array* OutputPatchElementIndex = new(ParseState)ir_dereference_array(
-				OutputPatchVar,
-				new(ParseState)ir_constant(
-				OutputVertex
-				)
-				);
-
-			ir_rvalue* OutputPatchElement = rhs->clone(ParseState, 0);
-			Helper::ReplaceVariableDerefWithArrayDeref(OutputPatchElement, OutputPatchElementIndex);
-
-			if (lhs)
-			{
-				// Wrap LHS into a record again
-				PostCallInstructions.push_tail(
-					new (ParseState)ir_assignment(
-						OutputPatchElement,
-						new(ParseState)ir_dereference_record(
-							new(ParseState)ir_dereference_array(
-								OutputPatchArray->clone(ParseState, 0),
-								new(ParseState)ir_constant(OutputVertex)
-							),
-							lhs->field
-						)
-					)
-				);
-			}
-			else
-			{
-				PostCallInstructions.push_tail(
-					new (ParseState)ir_assignment(
-						OutputPatchElement,
-						new(ParseState)ir_dereference_array(
-							OutputPatchArray->clone(ParseState, 0),
-							new(ParseState)ir_constant(OutputVertex)
-						)
-					)
-				);
-			}
-		}
-	}
-}
-
-void FVulkanLanguageSpec::SetupLanguageIntrinsics(_mesa_glsl_parse_state* State, exec_list* ir)
+	"GENERATED_SubpassDepthFetchAttachment",
+	"GENERATED_SubpassFetchAttachment0",
+	"GENERATED_SubpassFetchAttachment1",
+	"GENERATED_SubpassFetchAttachment2",
+	"GENERATED_SubpassFetchAttachment3",
+	"GENERATED_SubpassFetchAttachment4",
+	"GENERATED_SubpassFetchAttachment5",
+	"GENERATED_SubpassFetchAttachment6",
+	"GENERATED_SubpassFetchAttachment7"
+};
+const TCHAR* VULKAN_SUBPASS_FETCH_VAR_W[9] =
 {
-	auto AddIntrisicReturningFloat = [](_mesa_glsl_parse_state* State, exec_list* ir, const char* Name, int32 NumColumns)
-	{
-		ir_function* Func = new(State) ir_function(Name);
-		auto* ReturnType = glsl_type::get_instance(GLSL_TYPE_FLOAT, NumColumns, 1);
-		ir_function_signature* Sig = new(State) ir_function_signature(ReturnType);
-		Sig->is_builtin = true;
-		Func->add_signature(Sig);
-		State->symbols->add_global_function(Func);
-		ir->push_head(Func);
-	};
-
-	// Depth fetch
-	AddIntrisicReturningFloat(State, ir, VULKAN_SUBPASS_FETCH[0], 1);
-	// Color fetch
-	for (int32 i = 1; i < UE_ARRAY_COUNT(VULKAN_SUBPASS_FETCH); ++i)	
-	{
-		AddIntrisicReturningFloat(State, ir, VULKAN_SUBPASS_FETCH[i], 4);
-	}
-	
-	//if (State->language_version >= 310)
-	{
-		/**
-		* Create GLSL functions that are left out of the symbol table
-		*  Prevent pollution, but make them so thay can be used to
-		*  implement the hlsl barriers
-		*/
-		const int glslFuncCount = 7;
-		const char * glslFuncName[glslFuncCount] =
-		{
-			"barrier", "memoryBarrier", "memoryBarrierAtomicCounter", "memoryBarrierBuffer",
-			"memoryBarrierShared", "memoryBarrierImage", "groupMemoryBarrier"
-		};
-		ir_function* glslFuncs[glslFuncCount];
-
-		for (int i = 0; i < glslFuncCount; i++)
-		{
-			void* ctx = State;
-			ir_function* func = new(ctx)ir_function(glslFuncName[i]);
-			ir_function_signature* sig = new(ctx)ir_function_signature(glsl_type::void_type);
-			sig->is_builtin = true;
-			func->add_signature(sig);
-			ir->push_tail(func);
-			glslFuncs[i] = func;
-		}
-
-		/** Implement HLSL barriers in terms of GLSL functions */
-		const char * functions[] =
-		{
-			"GroupMemoryBarrier", "GroupMemoryBarrierWithGroupSync",
-			"DeviceMemoryBarrier", "DeviceMemoryBarrierWithGroupSync",
-			"AllMemoryBarrier", "AllMemoryBarrierWithGroupSync"
-		};
-		const int max_children = 4;
-		ir_function * implFuncs[][max_children] =
-		{
-			{ glslFuncs[4] } /**{"memoryBarrierShared"}*/,
-			{ glslFuncs[4], glslFuncs[0] } /**{"memoryBarrierShared","barrier"}*/,
-			{ glslFuncs[2], glslFuncs[3], glslFuncs[5] } /**{"memoryBarrierAtomicCounter", "memoryBarrierBuffer", "memoryBarrierImage"}*/,
-			{ glslFuncs[2], glslFuncs[3], glslFuncs[5], glslFuncs[0] } /**{"memoryBarrierAtomicCounter", "memoryBarrierBuffer", "memoryBarrierImage", "barrier"}*/,
-			{ glslFuncs[1] } /**{"memoryBarrier"}*/,
-			{ glslFuncs[1], glslFuncs[0] } /**{"groupMemoryBarrier","barrier"}*/
-		};
-
-		for (size_t i = 0; i < sizeof(functions) / sizeof(const char*); i++)
-		{
-			void* ctx = State;
-			ir_function* func = new(ctx)ir_function(functions[i]);
-
-			ir_function_signature* sig = new(ctx)ir_function_signature(glsl_type::void_type);
-			sig->is_builtin = true;
-			sig->is_defined = true;
-
-			for (int j = 0; j < max_children; j++)
-			{
-				if (implFuncs[i][j] == NULL)
-					break;
-				ir_function* child = implFuncs[i][j];
-				check(child);
-				check(child->signatures.get_head() == child->signatures.get_tail());
-				ir_function_signature *childSig = (ir_function_signature *)child->signatures.get_head();
-				exec_list actual_parameter;
-				sig->body.push_tail(
-					new(ctx)ir_call(childSig, NULL, &actual_parameter)
-					);
-			}
-
-			func->add_signature(sig);
-
-			State->symbols->add_global_function(func);
-			ir->push_tail(func);
-		}
-	}
-}
-
+	TEXT("GENERATED_SubpassDepthFetchAttachment"),
+	TEXT("GENERATED_SubpassFetchAttachment0"),
+	TEXT("GENERATED_SubpassFetchAttachment1"),
+	TEXT("GENERATED_SubpassFetchAttachment2"),
+	TEXT("GENERATED_SubpassFetchAttachment3"),
+	TEXT("GENERATED_SubpassFetchAttachment4"),
+	TEXT("GENERATED_SubpassFetchAttachment5"),
+	TEXT("GENERATED_SubpassFetchAttachment6"),
+	TEXT("GENERATED_SubpassFetchAttachment7")
+};
 
 FVulkanBindingTable::FBinding::FBinding()
 {
@@ -6218,106 +148,3 @@ void FVulkanBindingTable::SortBindings()
 		});
 	bSorted = true;
 }
-
-void FVulkanBindingTable::PrintBindingTableDefines(char** OutBuffer) const
-{
-	auto GetVulkanBindingTypeName = [](EVulkanBindingType::EType Type)
-		{
-			switch(Type)
-			{
-			case EVulkanBindingType::InputAttachment:		return "Input Attachments";
-			case EVulkanBindingType::PackedUniformBuffer:	return "Packed UB";
-			case EVulkanBindingType::UniformBuffer:			return "Uniform Buffer";
-			case EVulkanBindingType::CombinedImageSampler:	return "Combined Image Sampler";
-			case EVulkanBindingType::Sampler:				return "Sampler";
-			case EVulkanBindingType::Image:					return "Image";
-			case EVulkanBindingType::UniformTexelBuffer:	return "Uniform Texel Buffer";
-			case EVulkanBindingType::StorageImage:			return "Storage Image";
-			case EVulkanBindingType::StorageTexelBuffer:	return "Storage TexelBuffer";
-			case EVulkanBindingType::StorageBuffer:			return "Storage Buffer";
-			default:										return "INVALID!";
-			}
-		};
-	EVulkanBindingType::EType PreviousType = EVulkanBindingType::Count;
-	ralloc_asprintf_append(OutBuffer, "\n");
-	for (int32 Index = 0; Index < Bindings.Num(); ++Index)
-	{
-		if (PreviousType != Bindings[Index].Type)
-		{
-			ralloc_asprintf_append(OutBuffer, "// %s\n", GetVulkanBindingTypeName(Bindings[Index].Type));
-			PreviousType = Bindings[Index].Type;
-		}
-		ralloc_asprintf_append(OutBuffer, "#define BINDING_%d\t%d\n", Bindings[Index].VirtualIndex, Index);
-	}
-	ralloc_asprintf_append(OutBuffer, "\n");
-}
-
-struct FFixIntrinsicsVisitor : public ir_rvalue_visitor
-{
-	_mesa_glsl_parse_state* State;
-	//bool bUsesFramebufferFetchES2;
-	//int MRTFetchMask;
-	//ir_variable* DestColorVar;
-	//const glsl_type* DestColorType;
-	//ir_variable* DestMRTColorVar[MAX_SIMULTANEOUS_RENDER_TARGETS];
-
-	FFixIntrinsicsVisitor(_mesa_glsl_parse_state* InState) :
-		State(InState)
-		//bUsesFramebufferFetchES2(false),
-		//MRTFetchMask(0),
-		//DestColorVar(nullptr),
-		//DestColorType(glsl_type::error_type)
-	{
-
-	}
-
-	virtual void handle_rvalue(ir_rvalue** RValue)
-	{
-		if (!RValue || !*RValue)
-		{
-			return;
-		}
-
-		auto* expr = (*RValue)->as_expression();
-		if (!expr)
-		{
-			return;
-		}
-
-		ir_expression_operation op = expr->operation;
-
-		// Convert matrixCompMult to memberwise multiply
-		// and
-		// Convert binary matrix add to memberwise add
-		if ((op == ir_binop_mul || op == ir_binop_add ) 
-			&& expr->type->is_matrix()
-			&& expr->operands[0]->type->is_matrix()
-			&& expr->operands[1]->type->is_matrix())
-		{
-			check(expr->operands[0]->type == expr->operands[1]->type);
-			auto* NewTemp = new(State)ir_variable(expr->operands[0]->type, nullptr, ir_var_temporary);
-			base_ir->insert_before(NewTemp);
-			for (uint32 Index = 0; Index < expr->operands[0]->type->matrix_columns; ++Index)
-			{
-				auto* NewMul = new(State)ir_expression(op,
-					new(State)ir_dereference_array(expr->operands[0], new(State)ir_constant(Index)),
-					new(State)ir_dereference_array(expr->operands[1], new(State)ir_constant(Index)));
-				auto* NewAssign = new(State)ir_assignment(
-					new(State)ir_dereference_array(NewTemp, new(State)ir_constant(Index)),
-					NewMul);
-				base_ir->insert_before(NewAssign);
-			}
-
-			*RValue = new(State)ir_dereference_variable(NewTemp);
-		}
-	}
-};
-
-void FVulkanCodeBackend::FixIntrinsics(_mesa_glsl_parse_state* State, exec_list* ir)
-{
-	ir_function_signature* MainSig = GetMainFunction(ir);
-	check(MainSig);
-
-	FFixIntrinsicsVisitor Visitor(State);
-	Visitor.run(&MainSig->body);
-}
diff --git a/Engine/Source/Developer/VulkanShaderFormat/Private/VulkanBackend.h b/Engine/Source/Developer/VulkanShaderFormat/Private/VulkanBackend.h
--- a/Engine/Source/Developer/VulkanShaderFormat/Private/VulkanBackend.h
+++ b/Engine/Source/Developer/VulkanShaderFormat/Private/VulkanBackend.h
@@ -7,47 +7,6 @@
 #include "LanguageSpec.h"
 #include "VulkanCommon.h"
 
-class FVulkanLanguageSpec : public ILanguageSpec
-{
-protected:
-	bool bShareSamplers;
-	bool bRequiresOESExtensions;
-
-public:
-	FVulkanLanguageSpec(bool bInShareSamplers, bool bInRequiresOESExtensions)
-		: bShareSamplers(bInShareSamplers)
-		, bRequiresOESExtensions(bInRequiresOESExtensions)
-	{}
-
-	virtual bool SupportsDeterminantIntrinsic() const override
-	{
-		return true;
-	}
-
-	virtual bool SupportsTransposeIntrinsic() const override
-	{
-		return true;
-	}
-
-	virtual bool SupportsIntegerModulo() const override
-	{
-		return true;
-	}
-
-	virtual bool SupportsMatrixConversions() const override { return true; }
-
-	virtual void SetupLanguageIntrinsics(_mesa_glsl_parse_state* State, exec_list* ir) override;
-
-	virtual bool AllowsSharingSamplers() const override { return bShareSamplers; }
-
-	virtual bool RequiresNegateDDY() const override { return false; }
-
-public:
-	bool RequiresOESExtensions() const { return bRequiresOESExtensions; }
-};
-
-class ir_variable;
-
 // Generates Vulkan compliant code from IR tokens
 #ifdef __GNUC__
 #pragma GCC visibility push(default)
@@ -77,7 +36,6 @@ struct FVulkanBindingTable
 	}
 
 	void SortBindings();
-	void PrintBindingTableDefines(char** OutBuffer) const;
 
 	int32 GetRealBindingIndex(int32 InVirtualIndex) const
 	{
@@ -102,46 +60,6 @@ private:
 	TArray<FBinding> Bindings;
 
 	bool bSorted = false;
-
-	friend class FGenerateVulkanVisitor;
-};
-
-struct FVulkanCodeBackend : public FCodeBackend
-{
-	FVulkanCodeBackend(	unsigned int InHlslCompileFlags,
-						FVulkanBindingTable& InBindingTable,
-						EHlslCompileTarget InTarget) :
-		FCodeBackend(InHlslCompileFlags, InTarget),
-		BindingTable(InBindingTable),
-		bExplicitDepthWrites(false)
-	{
-	}
-
-	virtual char* GenerateCode(struct exec_list* ir, struct _mesa_glsl_parse_state* ParseState, EHlslShaderFrequency Frequency) override;
-
-	// Return false if there were restrictions that made compilation fail
-	virtual bool ApplyAndVerifyPlatformRestrictions(exec_list* Instructions, _mesa_glsl_parse_state* ParseState, EHlslShaderFrequency Frequency) override;
-
-	/**
-	* Generate a GLSL main() function that calls the entry point and handles
-	* reading and writing all input and output semantics.
-	* @param Frequency - The shader frequency.
-	* @param EntryPoint - The name of the shader entry point.
-	* @param Instructions - IR code.
-	* @param ParseState - Parse state.
-	*/
-	virtual bool GenerateMain(EHlslShaderFrequency Frequency, const char* EntryPoint, exec_list* Instructions, _mesa_glsl_parse_state* ParseState) override;
-
-	void FixIntrinsics(_mesa_glsl_parse_state* ParseState, exec_list* ir);
-
-	void GenShaderPatchConstantFunctionInputs(_mesa_glsl_parse_state* ParseState, ir_variable* OutputPatchVar, exec_list &PostCallInstructions);
-
-	void CallPatchConstantFunction(_mesa_glsl_parse_state* ParseState, ir_variable* OutputPatchVar, ir_function_signature* PatchConstantSig, exec_list& DeclInstructions, exec_list &PostCallInstructions);
-
-	ir_function_signature* FindPatchConstantFunction(exec_list* Instructions, _mesa_glsl_parse_state* ParseState);
-
-	FVulkanBindingTable& BindingTable;
-	bool bExplicitDepthWrites;
 };
 
 // InputAttachments
diff --git a/Engine/Source/Developer/VulkanShaderFormat/Private/VulkanShaderCompiler.cpp b/Engine/Source/Developer/VulkanShaderFormat/Private/VulkanShaderCompiler.cpp
--- a/Engine/Source/Developer/VulkanShaderFormat/Private/VulkanShaderCompiler.cpp
+++ b/Engine/Source/Developer/VulkanShaderFormat/Private/VulkanShaderCompiler.cpp
@@ -40,99 +40,6 @@ enum VkDescriptorType {
 
 DEFINE_LOG_CATEGORY_STATIC(LogVulkanShaderCompiler, Log, All); 
 
-static TArray<ANSICHAR> ParseIdentifierANSI(const FString& Str)
-{
-	TArray<ANSICHAR> Result;
-	Result.Reserve(Str.Len());
-	for (int32 Index = 0; Index < Str.Len(); ++Index)
-	{
-		Result.Add(FChar::ToLower((ANSICHAR)Str[Index]));
-	}
-	Result.Add('\0');
-
-	return Result;
-}
-
-
-inline const ANSICHAR * CStringEndOfLine(const ANSICHAR * Text)
-{
-	const ANSICHAR * LineEnd = FCStringAnsi::Strchr(Text, '\n');
-	if (nullptr == LineEnd)
-	{
-		LineEnd = Text + FCStringAnsi::Strlen(Text);
-	}
-	return LineEnd;
-}
-
-inline bool CStringIsBlankLine(const ANSICHAR * Text)
-{
-	while (!FCharAnsi::IsLinebreak(*Text))
-	{
-		if (!FCharAnsi::IsWhitespace(*Text))
-		{
-			return false;
-		}
-		++Text;
-	}
-	return true;
-}
-
-inline void AppendCString(TArray<ANSICHAR> & Dest, const ANSICHAR * Source)
-{
-	if (Dest.Num() > 0)
-	{
-		Dest.Insert(Source, FCStringAnsi::Strlen(Source), Dest.Num() - 1);;
-	}
-	else
-	{
-		Dest.Append(Source, FCStringAnsi::Strlen(Source) + 1);
-	}
-}
-
-inline bool MoveHashLines(TArray<ANSICHAR> & Dest, TArray<ANSICHAR> & Source)
-{
-	// Walk through the lines to find the first non-# line...
-	const ANSICHAR * LineStart = Source.GetData();
-	for (bool FoundNonHashLine = false; !FoundNonHashLine;)
-	{
-		const ANSICHAR * LineEnd = CStringEndOfLine(LineStart);
-		if (LineStart[0] != '#' && !CStringIsBlankLine(LineStart))
-		{
-			FoundNonHashLine = true;
-		}
-		else if (LineEnd[0] == '\n')
-		{
-			LineStart = LineEnd + 1;
-		}
-		else
-		{
-			LineStart = LineEnd;
-		}
-	}
-	// Copy the hash lines over, if we found any. And delete from
-	// the source.
-	if (LineStart > Source.GetData())
-	{
-		int32 LineLength = LineStart - Source.GetData();
-		if (Dest.Num() > 0)
-		{
-			Dest.Insert(Source.GetData(), LineLength, Dest.Num() - 1);
-		}
-		else
-		{
-			Dest.Append(Source.GetData(), LineLength);
-			Dest.Append("", 1);
-		}
-		if (Dest.Last(1) != '\n')
-		{
-			Dest.Insert("\n", 1, Dest.Num() - 1);
-		}
-		Source.RemoveAt(0, LineStart - Source.GetData());
-		return true;
-	}
-	return false;
-}
-
 static bool Match(const ANSICHAR* &Str, ANSICHAR Char)
 {
 	if (*Str == Char)
@@ -198,54 +105,6 @@ uint32 ParseNumber(const T* Str, bool bEmptyIsZero = false)
 	return Num;
 }
 
-static inline FString GetExtension(EHlslShaderFrequency Frequency, bool bAddDot = true)
-{
-	const TCHAR* Name = nullptr;
-	switch (Frequency)
-	{
-	default:
-		check(0);
-		// fallthrough...
-
-	case HSF_PixelShader:		Name = TEXT(".frag"); break;
-	case HSF_VertexShader:		Name = TEXT(".vert"); break;
-	case HSF_ComputeShader:		Name = TEXT(".comp"); break;
-	case HSF_GeometryShader:	Name = TEXT(".geom"); break;
-	case HSF_HullShader:		Name = TEXT(".tesc"); break;
-	case HSF_DomainShader:		Name = TEXT(".tese"); break;
-	}
-
-	if (!bAddDot)
-	{
-		++Name;
-	}
-	return FString(Name);
-}
-
-static uint32 GetTypeComponents(const FString& Type)
-{
-	static const FString TypePrefix[] = { "f", "i", "u" };
-	uint32 Components = 0;
-	int32 PrefixLength = 0;
-	for (uint32 i = 0; i<UE_ARRAY_COUNT(TypePrefix); i++)
-	{
-		const FString& Prefix = TypePrefix[i];
-		const int32 CmpLength = Type.Contains(Prefix, ESearchCase::CaseSensitive, ESearchDir::FromStart);
-		if (CmpLength == Prefix.Len())
-		{
-			PrefixLength = CmpLength;
-			break;
-		}
-	}
-
-	check(PrefixLength > 0);
-	Components = ParseNumber(*Type + PrefixLength);
-
-	check(Components > 0);
-	return Components;
-}
-
-
 static bool ContainsBinding(const FVulkanBindingTable& BindingTable, const FString& Name)
 {
 	for (const FVulkanBindingTable::FBinding& Binding : BindingTable.GetBindings())
@@ -1071,17 +930,6 @@ static void BuildShaderOutput(
 				OLDHeader.SerializedBindings.InOutMask.EnableField(Index + AttributeIndex);
 			}
 		}
-#if 0
-		// Record user-defined input varyings
-		else if (!Input.Name.StartsWith(GL_Prefix))
-		{
-			FVulkanShaderVarying Var;
-			Var.Location = Input.Index;
-			Var.Varying = ParseIdentifierANSI(Input.Name);
-			Var.Components = GetTypeComponents(Input.Type);
-			Header.SerializedBindings.InputVaryings.Add(Var);
-		}
-#endif
 	}
 
 	static const FString TargetPrefix = "out_Target";
@@ -1099,17 +947,6 @@ static void BuildShaderOutput(
 		{
 			OLDHeader.SerializedBindings.InOutMask.EnableField(CrossCompiler::FShaderBindingInOutMask::DepthStencilMaskIndex);
 		}
-#if 0
-		// Record user-defined output varyings
-		else if (!Output.Name.StartsWith(GL_Prefix))
-		{
-			FVulkanShaderVarying Var;
-			Var.Location = Output.Index;
-			Var.Varying = ParseIdentifierANSI(Output.Name);
-			Var.Components = GetTypeComponents(Output.Type);
-			Header.SerializedBindings.OutputVaryings.Add(Var);
-		}
-#endif
 	}
 
 
@@ -1522,161 +1359,6 @@ static void BuildShaderOutput(
 	CullGlobalUniformBuffers(ShaderInput.Environment.UniformBufferMap, ShaderOutput.ParameterMap);
 }
 
-static char* PatchGLSLVersionPosition(const char* InSourceGLSL)
-{
-	if(!InSourceGLSL)
-	{
-		return nullptr;
-	}
-
-	const int32 InSrcLength = FCStringAnsi::Strlen(InSourceGLSL);
-	if(InSrcLength <= 0)
-	{
-		return nullptr;
-	}
-
-	char* GlslSource = (char*)malloc(InSrcLength+1);
-	check(GlslSource);
-	memcpy(GlslSource, InSourceGLSL, InSrcLength+1);
-
-	// Find begin of "#version" line
-	char* VersionBegin = strstr(GlslSource, "#version");
-
-	// Find end of "#version line"
-	char* VersionEnd = VersionBegin ? strstr(VersionBegin, "\n") : nullptr;
-
-	if(VersionEnd)
-	{
-		// Add '\n' character
-		VersionEnd++;
-
-		const int32 VersionLineLength = VersionEnd - VersionBegin - 1;
-
-		// Copy version line into a temporary buffer (+1 for term-char).
-		const int32 TmpStrBytes = (VersionEnd - VersionBegin) + 1;
-		char* TmpVersionLine = (char*)malloc(TmpStrBytes);
-		check(TmpVersionLine);
-		memset(TmpVersionLine, 0, TmpStrBytes);
-		memcpy(TmpVersionLine, VersionBegin, VersionEnd - VersionBegin);
-
-		// Erase current version number, just replace it with spaces...
-		for(char* str=VersionBegin; str<(VersionEnd-1); str++)
-		{
-			*str=' ';
-		}
-
-		// Allocate new source buffer to place version string on the first line.
-		char* NewSource = (char*)malloc(InSrcLength + TmpStrBytes);
-		check(NewSource);
-
-		// Copy version line
-		memcpy(NewSource, TmpVersionLine, TmpStrBytes);
-
-		// Copy original source after the source line
-		// -1 to offset back from the term-char
-		memcpy(NewSource + TmpStrBytes - 1, GlslSource, InSrcLength + 1);
-
-		free(TmpVersionLine);
-		TmpVersionLine = nullptr;
-							
-		// Update string pointer
-		free(GlslSource);
-		GlslSource = NewSource;
-	}
-
-	return GlslSource;
-}
-
-static void PatchForToWhileLoop(char** InOutSourceGLSL)
-{
-	//checkf(InOutSourceGLSL, TEXT("Attempting to patch an invalid glsl source-string"));
-
-	char* srcGlsl = *InOutSourceGLSL;
-	//checkf(srcGlsl, TEXT("Attempting to patch an invalid glsl source-string"));
-
-	const size_t InSrcLength = strlen(srcGlsl);
-	//checkf(InSrcLength > 0, TEXT("Attempting to patch an empty glsl source-string."));
-
-	// This is what we are relacing
-	const char* srcPatchable = "for (;;)";
-	const size_t srcPatchableLength = strlen(srcPatchable);
-
-	// This is where we are replacing with
-	const char* dstPatchable = "while(true)";
-	const size_t dstPatchableLength = strlen(dstPatchable);
-	
-	// Find number of occurances
-	int numNumberOfOccurances = 0;
-	for(char* dstReplacePos = strstr(srcGlsl, srcPatchable);
-		dstReplacePos != NULL;
-		dstReplacePos = strstr(dstReplacePos+srcPatchableLength, srcPatchable))
-	{
-		numNumberOfOccurances++;
-	}
-
-	// No patching needed
-	if(numNumberOfOccurances == 0)
-	{
-		return;
-	}
-
-	// Calc new required string-length
-	const size_t newLength = InSrcLength + (dstPatchableLength-srcPatchableLength)*numNumberOfOccurances;
-
-	// Allocate destination buffer + 1 char for terminating character
-	char* GlslSource = (char*)malloc(newLength+1);
-	check(GlslSource);
-	memset(GlslSource, 0, sizeof(char)*(newLength+1));
-	memcpy(GlslSource, srcGlsl, InSrcLength);
-
-	// Scan and replace
-	char* dstReplacePos = strstr(GlslSource, srcPatchable);
-	char* srcReplacePos = strstr(srcGlsl, srcPatchable);
-	int bytesRemaining = (int)newLength;
-	while(dstReplacePos != NULL && srcReplacePos != NULL)
-	{
-		// Replace the string
-		bytesRemaining = (int)newLength - (int)(dstReplacePos - GlslSource);
-		memcpy(dstReplacePos, dstPatchable, dstPatchableLength);
-		
-		// Increment positions
-		dstReplacePos+=dstPatchableLength;
-		srcReplacePos+=srcPatchableLength;
-
-		// Append remaining code
-		int bytesToCopy = InSrcLength - (int)(srcReplacePos - srcGlsl);
-		memcpy(dstReplacePos, srcReplacePos, bytesToCopy);
-
-		dstReplacePos = strstr(dstReplacePos, srcPatchable);
-		srcReplacePos = strstr(srcReplacePos, srcPatchable);
-	}
-
-	free(*InOutSourceGLSL);
-
-	*InOutSourceGLSL = GlslSource;
-}
-
-static FString CreateShaderCompileCommandLine(FCompilerInfo& CompilerInfo, EHlslCompileTarget Target)
-{
-	FString CmdLine;
-	FString GLSLFile = CompilerInfo.Input.DumpDebugInfoPath / (TEXT("Output") + GetExtension(CompilerInfo.Frequency));
-	FString SPVFile = CompilerInfo.Input.DumpDebugInfoPath / TEXT("Output.spv");
-	FString SPVDisasmFile = CompilerInfo.Input.DumpDebugInfoPath / TEXT("Output.spvasm");
-
-	CmdLine += TEXT("\n\"");
-#if PLATFORM_WINDOWS
-	CmdLine += *(FPaths::RootDir() / TEXT("Engine/Binaries/ThirdParty/glslang/glslangValidator.exe"));
-#elif PLATFORM_LINUX
-	CmdLine += *(FPaths::RootDir() / TEXT("Engine/Binaries/ThirdParty/glslang/glslangValidator"));
-#endif
-	CmdLine += TEXT("\"");
-	CmdLine += TEXT(" -V -H -r -o \"") + SPVFile + TEXT("\" \"") + GLSLFile + TEXT("\" > \"" + SPVDisasmFile + "\"");
-	CmdLine += TEXT("\npause\n");
-
-	return CmdLine;
-}
-
-
 FCompilerInfo::FCompilerInfo(const FShaderCompilerInput& InInput, const FString& InWorkingDirectory, EHlslShaderFrequency InFrequency) :
 	Input(InInput),
 	WorkingDirectory(InWorkingDirectory),
@@ -1750,84 +1432,6 @@ static bool CompileUsingInternal(FCompilerInfo& CompilerInfo, const FVulkanBindi
 }
 
 
-static bool CompileWithHlslcc(const FString& PreprocessedShader, FVulkanBindingTable& BindingTable, FCompilerInfo& CompilerInfo, FString& EntryPointName, EHlslCompileTarget HlslCompilerTarget, FShaderCompilerOutput& Output, TArray<ANSICHAR>& OutGlsl)
-{
-	char* GlslShaderSource = nullptr;
-	char* ErrorLog = nullptr;
-
-	auto InnerFunction = [&]()
-	{
-		// Call hlslcc
-		FVulkanCodeBackend VulkanBackend(CompilerInfo.CCFlags, BindingTable, HlslCompilerTarget);
-		FHlslCrossCompilerContext CrossCompilerContext(CompilerInfo.CCFlags, CompilerInfo.Frequency, HlslCompilerTarget);
-
-		const bool bShareSamplers = true;
-		const bool bRequiresOESExtensions = true;
-
-		FVulkanLanguageSpec VulkanLanguageSpec(bShareSamplers, bRequiresOESExtensions);
-		int32 Result = 0;
-		if (CrossCompilerContext.Init(TCHAR_TO_ANSI(*CompilerInfo.Input.VirtualSourceFilePath), &VulkanLanguageSpec))
-		{
-			Result = CrossCompilerContext.Run(
-				TCHAR_TO_ANSI(*PreprocessedShader),
-				TCHAR_TO_ANSI(*EntryPointName),
-				&VulkanBackend,
-				&GlslShaderSource,
-				&ErrorLog
-				) ? 1 : 0;
-		}
-
-		if (Result == 0)
-		{
-			FString Tmp = ANSI_TO_TCHAR(ErrorLog);
-			TArray<FString> ErrorLines;
-			Tmp.ParseIntoArray(ErrorLines, TEXT("\n"), true);
-			for (int32 LineIndex = 0; LineIndex < ErrorLines.Num(); ++LineIndex)
-			{
-				const FString& Line = ErrorLines[LineIndex];
-				CrossCompiler::ParseHlslccError(Output.Errors, Line, CompilerInfo.Input.bSkipPreprocessedCache);
-			}
-
-			return false;
-		}
-
-		check(GlslShaderSource);
-
-		// Patch GLSL source
-		PatchForToWhileLoop(&GlslShaderSource);
-
-		if (CompilerInfo.bDebugDump)
-		{
-			FString DumpedGlslFile = CompilerInfo.Input.DumpDebugInfoPath / (TEXT("Output") + GetExtension(CompilerInfo.Frequency));
-			if (TUniquePtr<FArchive> FileWriter = TUniquePtr<FArchive>(IFileManager::Get().CreateFileWriter(*DumpedGlslFile)))
-			{
-				FileWriter->Serialize(GlslShaderSource, FCStringAnsi::Strlen(GlslShaderSource));
-				FileWriter->Close();
-			}
-		}
-
-		int32 Length = FCStringAnsi::Strlen(GlslShaderSource);
-		OutGlsl.AddUninitialized(Length + 1);
-		FCStringAnsi::Strcpy(OutGlsl.GetData(), Length + 1, GlslShaderSource);
-
-		return true;
-	};
-
-	bool bResult = InnerFunction();
-
-	if (ErrorLog)
-	{
-		free(ErrorLog);
-	}
-
-	if (GlslShaderSource)
-	{
-		free(GlslShaderSource);
-	}
-
-	return bResult;
-}
-
 #if PLATFORM_MAC || PLATFORM_WINDOWS || PLATFORM_LINUX
 
 static void GatherSpirvReflectionBindings(
@@ -2429,7 +2033,6 @@ void DoCompileVulkanShader(const FShaderCompilerInput& Input, FShaderCompilerOut
 	const bool bIsSM5 = (Version == EVulkanShaderVersion::SM5);
 	const bool bIsMobile = (Version == EVulkanShaderVersion::ES3_1 || Version == EVulkanShaderVersion::ES3_1_ANDROID);
 	const bool bStripReflect = (IsVulkanMobilePlatform(ShaderPlatform) || IsVulkanMobileSM5Platform(ShaderPlatform)) || Input.IsRayTracingShader();
-	const bool bForceDXC = Input.Environment.CompilerFlags.Contains(CFLAG_ForceDXC);
 
 	const EHlslShaderFrequency FrequencyTable[] =
 	{
@@ -2570,39 +2173,9 @@ void DoCompileVulkanShader(const FShaderCompilerInput& Input, FShaderCompilerOut
 	bool bSuccess = false;
 
 #if PLATFORM_MAC || PLATFORM_WINDOWS || PLATFORM_LINUX
-	if (bForceDXC)
-	{
-		// Cross-compile shader via ShaderConductor (DXC, SPIRV-Tools, SPIRV-Cross)
-		bSuccess = CompileWithShaderConductor(PreprocessedShaderSource, EntryPointName, CompilerInfo, HlslCompilerTarget, Output, BindingTable, bStripReflect);
-	}
-	else
+	// Cross-compile shader via ShaderConductor (DXC, SPIRV-Tools, SPIRV-Cross)
+	bSuccess = CompileWithShaderConductor(PreprocessedShaderSource, EntryPointName, CompilerInfo, HlslCompilerTarget, Output, BindingTable, bStripReflect);
 #endif // PLATFORM_MAC || PLATFORM_WINDOWS || PLATFORM_LINUX
-	{
-		if (CompilerInfo.bDebugDump)
-		{
-			if (Input.bGenerateDirectCompileFile)
-			{
-				FFileHelper::SaveStringToFile(CreateShaderCompilerWorkerDirectCommandLine(Input), *(Input.DumpDebugInfoPath / TEXT("DirectCompile.txt")));
-			}
-
-			const FString BatchFileContents = CreateShaderCompileCommandLine(CompilerInfo, HlslCompilerTarget);
-			FFileHelper::SaveStringToFile(BatchFileContents, *(CompilerInfo.Input.DumpDebugInfoPath / TEXT("CompileSPIRV.bat")));
-		}
-
-		// Cross-compile shader via HLSLcc
-		if (CompileWithHlslcc(PreprocessedShaderSource, BindingTable, CompilerInfo, EntryPointName, HlslCompilerTarget, Output, GeneratedGlslSource))
-		{
-			// For debugging: if you hit an error from Glslang/Spirv, use the SourceNoHeader for line numbers
-			auto* SourceWithHeader = GeneratedGlslSource.GetData();
-			char* SourceNoHeader = strstr(SourceWithHeader, "#version");
-			bSuccess = CompileUsingInternal(CompilerInfo, BindingTable, GeneratedGlslSource, Output);
-			if (bDirectCompile)
-			{
-				FPlatformMisc::LowLevelOutputDebugStringf(TEXT("Success: %d\n%s\n"), bSuccess, ANSI_TO_TCHAR(SourceWithHeader));
-			}
-		}
-	}
-
 	
 	ShaderParameterParser.ValidateShaderParameterTypes(Input, Output);
 	
diff --git a/Engine/Source/Developer/VulkanShaderFormat/Private/VulkanShaderFormat.cpp b/Engine/Source/Developer/VulkanShaderFormat/Private/VulkanShaderFormat.cpp
--- a/Engine/Source/Developer/VulkanShaderFormat/Private/VulkanShaderFormat.cpp
+++ b/Engine/Source/Developer/VulkanShaderFormat/Private/VulkanShaderFormat.cpp
@@ -26,7 +26,7 @@ class FShaderFormatVulkan : public IShaderFormat
 
 	int32 InternalGetVersion(FName Format) const
 	{
-		if (Format == NAME_VULKAN_SM5 ||	Format == NAME_VULKAN_SM5_ANDROID)
+		if (Format == NAME_VULKAN_SM5 || Format == NAME_VULKAN_SM5_ANDROID)
 		{
 			return UE_SHADER_VULKAN_SM5_VER;
 		}
@@ -85,13 +85,6 @@ public:
 		}
 	}
 
-	//virtual bool CreateLanguage(FName Format, ILanguageSpec*& OutSpec, FCodeBackend*& OutBackend, uint32 InHlslCompileFlags) override
-	//{
-	//	OutSpec = new FVulkanLanguageSpec(false);
-	//	OutBackend = new FVulkanCodeBackend(InHlslCompileFlags, HCT_FeatureLevelSM4);
-	//	return false;
-	//}
-
 	virtual const TCHAR* GetPlatformIncludeDirectory() const
 	{
 		return TEXT("Vulkan");
diff --git a/Engine/Source/Developer/VulkanShaderFormat/VulkanShaderFormat.Build.cs b/Engine/Source/Developer/VulkanShaderFormat/VulkanShaderFormat.Build.cs
--- a/Engine/Source/Developer/VulkanShaderFormat/VulkanShaderFormat.Build.cs
+++ b/Engine/Source/Developer/VulkanShaderFormat/VulkanShaderFormat.Build.cs
@@ -22,7 +22,6 @@ public class VulkanShaderFormat : ModuleRules
 			}
 			);
 
-		AddEngineThirdPartyPrivateStaticDependencies(Target, "HLSLCC");
 		AddEngineThirdPartyPrivateStaticDependencies(Target, "GlsLang");
 
 		if (Target.Platform == UnrealTargetPlatform.Mac || Target.Platform == UnrealTargetPlatform.Win64 || Target.Platform == UnrealTargetPlatform.Linux)
diff --git a/Engine/Source/Runtime/Engine/Private/ShaderCompiler/ShaderCompiler.cpp b/Engine/Source/Runtime/Engine/Private/ShaderCompiler/ShaderCompiler.cpp
--- a/Engine/Source/Runtime/Engine/Private/ShaderCompiler/ShaderCompiler.cpp
+++ b/Engine/Source/Runtime/Engine/Private/ShaderCompiler/ShaderCompiler.cpp
@@ -962,16 +962,6 @@ static TAutoConsoleVariable<int32> CVarOpenGLForceDXC(
 	TEXT(" 1: Force new compiler for all shaders (default)"),
 	ECVF_ReadOnly);
 
-static TAutoConsoleVariable<int32> CVarVulkanForceDXC(
-	TEXT("r.Vulkan.ForceDXC"),
-	3, 
-	TEXT("Forces DirectX Shader Compiler (DXC) to be used for all Vulkan shaders instead of hlslcc.\n")
-	TEXT(" 0: Disable (hlslcc/glslang)\n")
-	TEXT(" 1: Enabled on desktop platforms only (default)\n")
-	TEXT(" 2: Enabled on mobile platforms only\n")
-	TEXT(" 3: Enabled on all platforms"),
-	ECVF_ReadOnly);
-
 int32 GCreateShadersOnLoad = 0;
 static FAutoConsoleVariableRef CVarCreateShadersOnLoad(
 	TEXT("r.CreateShadersOnLoad"),
diff --git a/Engine/Source/Runtime/RenderCore/Private/Shader.cpp b/Engine/Source/Runtime/RenderCore/Private/Shader.cpp
--- a/Engine/Source/Runtime/RenderCore/Private/Shader.cpp
+++ b/Engine/Source/Runtime/RenderCore/Private/Shader.cpp
@@ -1302,21 +1302,11 @@ bool IsDxcEnabledForPlatform(EShaderPlatform Platform)
 		static const auto CVar = IConsoleManager::Get().FindConsoleVariable(TEXT("r.OpenGL.ForceDXC"));
 		return (CVar && CVar->GetInt() != 0);
 	}
-	if (IsMetalPlatform(Platform))
+	// Hlslcc has been removed for Metal and Vulkan. There is only DXC now.
+	if (IsMetalPlatform(Platform) || IsVulkanPlatform(Platform))
 	{
-		// Hlslcc has been removed for Metal. There is only DXC now.
 		return true;
 	}
-	if (IsVulkanPlatform(Platform))
-	{
-		static const auto CVar = IConsoleManager::Get().FindConsoleVariable(TEXT("r.Vulkan.ForceDXC"));
-		int32 VulkanForceDxc = (CVar ? CVar->GetInt() : 0);
-		const bool bIsVulkanMobile = IsVulkanMobilePlatform((EShaderPlatform)Platform) || IsVulkanMobileSM5Platform((EShaderPlatform)Platform);
-		const bool bIsDxcEnabledForDesktop = (VulkanForceDxc == 1 && !bIsVulkanMobile);
-		const bool bIsDxcEnabledForMobile = (VulkanForceDxc == 2 && bIsVulkanMobile);
-		const bool bIsDxcEnableForAll = (VulkanForceDxc == 3);
-		return (bIsDxcEnabledForDesktop || bIsDxcEnabledForMobile || bIsDxcEnableForAll);
-	}
 	return false;
 }
 
diff --git a/Engine/Source/Runtime/RenderCore/Public/Shader.h b/Engine/Source/Runtime/RenderCore/Public/Shader.h
--- a/Engine/Source/Runtime/RenderCore/Public/Shader.h
+++ b/Engine/Source/Runtime/RenderCore/Public/Shader.h
@@ -2359,7 +2359,7 @@ inline void DispatchComputeShader(
 /** Returns whether the platform is using emulated uniform buffers */
 extern RENDERCORE_API bool IsUsingEmulatedUniformBuffers(EShaderPlatform Platform);
 
-/** Returns whether DirectXShaderCompiler (DXC) is enabled for the specified shader platform. See console variables "r.OpenGL.ForceDXC", "r.Vulkan.ForceDXC", and "r.D3D.ForceDXC". */
+/** Returns whether DirectXShaderCompiler (DXC) is enabled for the specified shader platform. See console variables "r.OpenGL.ForceDXC", "r.D3D.ForceDXC". */
 extern RENDERCORE_API bool IsDxcEnabledForPlatform(EShaderPlatform Platform);
 
 /** Appends to KeyString for all shaders. */
diff --git a/Engine/Source/Runtime/VulkanRHI/Private/VulkanPendingState.cpp b/Engine/Source/Runtime/VulkanRHI/Private/VulkanPendingState.cpp
--- a/Engine/Source/Runtime/VulkanRHI/Private/VulkanPendingState.cpp
+++ b/Engine/Source/Runtime/VulkanRHI/Private/VulkanPendingState.cpp
@@ -508,7 +508,7 @@ void FVulkanPendingGfxState::InternalUpdateDynamicStates(FVulkanCmdBuffer* Cmd)
 	{
 		ensure(Viewport.width > 0 || Viewport.height > 0);
 
-		// Flip viewport on Y-axis to be uniform between HLSLcc and DXC generated SPIR-V shaders (requires VK_KHR_maintenance1 extension)
+		// Flip viewport on Y-axis to be uniform between DXC generated SPIR-V shaders (requires VK_KHR_maintenance1 extension)
 		VkViewport FlippedViewport = Viewport;
 		FlippedViewport.y += FlippedViewport.height;
 		FlippedViewport.height = -FlippedViewport.height;
