Removing new usage of DDSPI in a shader format.

diff --git a/Engine/Source/Developer/VulkanShaderFormat/Private/VulkanShaderCompiler.cpp b/Engine/Source/Developer/VulkanShaderFormat/Private/VulkanShaderCompiler.cpp
--- a/Engine/Source/Developer/VulkanShaderFormat/Private/VulkanShaderCompiler.cpp
+++ b/Engine/Source/Developer/VulkanShaderFormat/Private/VulkanShaderCompiler.cpp
@@ -59,6 +59,13 @@ inline bool SupportsOfflineCompiler(EShaderPlatform ShaderPlatform)
 		|| ShaderPlatform == SP_VULKAN_SM5_ANDROID;
 }
 
+inline bool ForceSubpassImageDepthFalseForPlatform(EShaderPlatform ShaderPlatform)
+{
+	return ShaderPlatform == SP_VULKAN_PCES3_1
+		|| ShaderPlatform == SP_VULKAN_ES3_1_ANDROID
+		|| ShaderPlatform == SP_VULKAN_SM5_ANDROID;
+}
+
 
 DEFINE_LOG_CATEGORY_STATIC(LogVulkanShaderCompiler, Log, All); 
 
@@ -2000,12 +2007,7 @@ static bool CompileWithShaderConductor(
 	// Initialize compilation options for ShaderConductor
 	CrossCompiler::FShaderConductorOptions Options;
 	Options.bDisableScalarBlockLayout = !bIsRayTracingShader;
-
-	EShaderPlatform TargetPlatform = Input.Target.GetPlatform();
-	if (IsVulkanMobilePlatform(TargetPlatform) || IsVulkanMobileSM5Platform(TargetPlatform))
-	{
-		Options.bForceSubpassImageDepthFalse = true;
-	}
+	Options.bForceSubpassImageDepthFalse = ForceSubpassImageDepthFalseForPlatform(Input.Target.GetPlatform());
 
 	// Ray tracing features require Vulkan 1.2 environment.
 	if (bIsRayTracingShader || Input.Environment.CompilerFlags.Contains(CFLAG_InlineRayTracing))
