Fixed: missing shader code in material editor for vulkan platforms

diff --git a/Engine/Source/Developer/VulkanShaderFormat/Private/VulkanShaderCompiler.cpp b/Engine/Source/Developer/VulkanShaderFormat/Private/VulkanShaderCompiler.cpp
--- a/Engine/Source/Developer/VulkanShaderFormat/Private/VulkanShaderCompiler.cpp
+++ b/Engine/Source/Developer/VulkanShaderFormat/Private/VulkanShaderCompiler.cpp
@@ -1375,9 +1375,11 @@ static void BuildShaderOutput(
 
 	if (ShaderInput.ExtraSettings.bExtractShaderSource)
 	{
-		TArray<ANSICHAR> CodeOriginal;
-		CodeOriginal.Append(USFSource, FCStringAnsi::Strlen(USFSource) + 1);
-		ShaderOutput.OptionalFinalShaderSource = FString(CodeOriginal.GetData());
+		TArray<ANSICHAR> AssemblyText;
+		if (CrossCompiler::FShaderConductorContext::Disassemble(CrossCompiler::EShaderConductorIR::Spirv, Spirv.GetByteData(), Spirv.GetByteSize(), AssemblyText))
+		{
+			ShaderOutput.OptionalFinalShaderSource = FString(AssemblyText.GetData());
+		}
 	}
 	if (ShaderInput.ExtraSettings.OfflineCompilerPath.Len() > 0)
 	{
