Vulkan: Use Vulkan 1.1 environment for shader compilation on desktop platforms.  Enable Wave Ops in Vulkan.

diff --git a/Engine/Shaders/Public/Platform/Vulkan/VulkanCommon.ush b/Engine/Shaders/Public/Platform/Vulkan/VulkanCommon.ush
--- a/Engine/Shaders/Public/Platform/Vulkan/VulkanCommon.ush
+++ b/Engine/Shaders/Public/Platform/Vulkan/VulkanCommon.ush
@@ -5,7 +5,7 @@
 =============================================================================*/
 
 // Update this GUID to improve shader recompilation for Vulkan only shaders
-// GUID = BF67D1A7-590E-4A06-BBA2-37F7B9A72D28
+// GUID = 4F5D9B6D-5648-4707-B604-FDAA6C9CC36B
 
 #pragma once
 
diff --git a/Engine/Source/Developer/VulkanShaderFormat/Private/VulkanShaderCompiler.cpp b/Engine/Source/Developer/VulkanShaderFormat/Private/VulkanShaderCompiler.cpp
--- a/Engine/Source/Developer/VulkanShaderFormat/Private/VulkanShaderCompiler.cpp
+++ b/Engine/Source/Developer/VulkanShaderFormat/Private/VulkanShaderCompiler.cpp
@@ -66,6 +66,13 @@ inline bool ForceSubpassImageDepthFalseForPlatform(EShaderPlatform ShaderPlatfor
 		|| ShaderPlatform == SP_VULKAN_SM5_ANDROID;
 }
 
+inline CrossCompiler::FShaderConductorOptions::ETargetEnvironment GetMinimumTargetEnvironment(EShaderPlatform ShaderPlatform)
+{
+	// All desktop Vulkan platforms create a Vulkan 1.1 instance (SP_VULKAN_SM5_ANDROID can still use 1.0)
+	return (ShaderPlatform == SP_VULKAN_SM5) ? 
+		CrossCompiler::FShaderConductorOptions::ETargetEnvironment::Vulkan_1_1: 
+		CrossCompiler::FShaderConductorOptions::ETargetEnvironment::Vulkan_1_0;
+}
 
 DEFINE_LOG_CATEGORY_STATIC(LogVulkanShaderCompiler, Log, All); 
 
@@ -2143,6 +2150,10 @@ static bool CompileWithShaderConductor(
 	{
 		Options.TargetEnvironment = CrossCompiler::FShaderConductorOptions::ETargetEnvironment::Vulkan_1_2;
 	}
+	else
+	{
+		Options.TargetEnvironment = GetMinimumTargetEnvironment(Input.Target.GetPlatform());
+	}
 
 	if (bRewriteHlslSource)
 	{
@@ -2203,6 +2214,7 @@ void DoCompileVulkanShader(const FShaderCompilerInput& Input, FShaderCompilerOut
 	const bool bIsSM5 = (Version == EVulkanShaderVersion::SM5);
 	const bool bIsMobile = (Version == EVulkanShaderVersion::ES3_1 || Version == EVulkanShaderVersion::ES3_1_ANDROID);
 	const bool bStripReflect = SupportsOfflineCompiler(Input.Target.GetPlatform()) || Input.IsRayTracingShader();
+	const CrossCompiler::FShaderConductorOptions::ETargetEnvironment TargetEnvironment = GetMinimumTargetEnvironment(Input.Target.GetPlatform());
 
 	const EHlslShaderFrequency FrequencyTable[] =
 	{
@@ -2235,7 +2247,7 @@ void DoCompileVulkanShader(const FShaderCompilerInput& Input, FShaderCompilerOut
 	FShaderCompilerDefinitions AdditionalDefines;
 	AdditionalDefines.SetDefine(TEXT("COMPILER_HLSLCC"), 1);
 	AdditionalDefines.SetDefine(TEXT("COMPILER_VULKAN"), 1);
-	if(bIsMobile)
+	if (bIsMobile)
 	{
 		AdditionalDefines.SetDefine(TEXT("ES3_1_PROFILE"), 1);
 		AdditionalDefines.SetDefine(TEXT("VULKAN_PROFILE"), 1);
@@ -2259,6 +2271,11 @@ void DoCompileVulkanShader(const FShaderCompilerInput& Input, FShaderCompilerOut
 		AdditionalDefines.SetDefine(TEXT("PLATFORM_SUPPORTS_INLINE_RAY_TRACING"), 1);
 	}
 
+	if (TargetEnvironment >= CrossCompiler::FShaderConductorOptions::ETargetEnvironment::Vulkan_1_1)
+	{
+		AdditionalDefines.SetDefine(TEXT("PLATFORM_SUPPORTS_SM6_0_WAVE_OPERATIONS"), 1);
+	}
+
 	// Preprocess the shader.
 	FString PreprocessedShaderSource;
 	const bool bDirectCompile = FParse::Param(FCommandLine::Get(), TEXT("directcompile"));
