PR #8976: Add a setting to disable strip-reflect when Android shaders are compiled (Contributed by [REDACTED_PERSON])

diff --git a/Engine/Config/BaseEngine.ini b/Engine/Config/BaseEngine.ini
--- a/Engine/Config/BaseEngine.ini
+++ b/Engine/Config/BaseEngine.ini
@@ -2662,6 +2662,7 @@ bBundleLanguageSplit=True
 bBundleDensitySplit=True
 bFullScreen=True
 bForceLDLinker=False
+bStripShaderReflection=True
 bForceSmallOBBFiles=False
 bAllowLargeOBBFiles=False
 bAllowPatchOBBFile=False
diff --git a/Engine/Source/Developer/VulkanShaderFormat/Private/VulkanShaderCompiler.cpp b/Engine/Source/Developer/VulkanShaderFormat/Private/VulkanShaderCompiler.cpp
--- a/Engine/Source/Developer/VulkanShaderFormat/Private/VulkanShaderCompiler.cpp
+++ b/Engine/Source/Developer/VulkanShaderFormat/Private/VulkanShaderCompiler.cpp
@@ -45,10 +45,9 @@ inline bool CanCompilePlatform(EShaderPlatform ShaderPlatform)
 		|| ShaderPlatform == SP_VULKAN_SM5_ANDROID;
 }
 
-inline bool ShouldStripReflection(EShaderPlatform ShaderPlatform)
+inline bool IsAndroidShaderPlatform(EShaderPlatform ShaderPlatform)
 {
-	return ShaderPlatform == SP_VULKAN_PCES3_1
-		|| ShaderPlatform == SP_VULKAN_ES3_1_ANDROID
+	return ShaderPlatform == SP_VULKAN_ES3_1_ANDROID
 		|| ShaderPlatform == SP_VULKAN_SM5_ANDROID;
 }
 
@@ -2219,7 +2218,14 @@ void DoCompileVulkanShader(const FShaderCompilerInput& Input, FShaderCompilerOut
 
 	const bool bIsSM5 = (Version == EVulkanShaderVersion::SM5);
 	const bool bIsMobile = (Version == EVulkanShaderVersion::ES3_1 || Version == EVulkanShaderVersion::ES3_1_ANDROID);
-	const bool bStripReflect = SupportsOfflineCompiler(Input.Target.GetPlatform()) || Input.IsRayTracingShader();
+	bool bStripReflect = Input.IsRayTracingShader();
+	// By default we strip reflecion information for Android platform to avoid issues with older drivers
+	if (IsAndroidShaderPlatform(Input.Target.GetPlatform()))
+	{
+		const FString* StripReflect_Android = Input.Environment.GetDefinitions().Find(TEXT("STRIP_REFLECT_ANDROID"));
+		bStripReflect = !(StripReflect_Android && *StripReflect_Android == TEXT("0"));
+	}
+
 	const CrossCompiler::FShaderConductorOptions::ETargetEnvironment TargetEnvironment = GetMinimumTargetEnvironment(Input.Target.GetPlatform());
 
 	const EHlslShaderFrequency FrequencyTable[] =
diff --git a/Engine/Source/Runtime/Android/AndroidRuntimeSettings/Classes/AndroidRuntimeSettings.h b/Engine/Source/Runtime/Android/AndroidRuntimeSettings/Classes/AndroidRuntimeSettings.h
--- a/Engine/Source/Runtime/Android/AndroidRuntimeSettings/Classes/AndroidRuntimeSettings.h
+++ b/Engine/Source/Runtime/Android/AndroidRuntimeSettings/Classes/AndroidRuntimeSettings.h
@@ -407,6 +407,13 @@ public:
 	UPROPERTY(GlobalConfig, EditAnywhere, Category = AdvancedBuild, meta = (DisplayName = "Force linking to use ld instead of lld"))
 	bool bForceLDLinker;
 
+	// Strip shader reflection information under Android to avoid issues on older drivers
+	UPROPERTY(config, EditAnywhere, Category = AdvancedBuild, meta = (
+	DisplayName = "Strip shader reflection information",
+	ToolTip = "If true, strip shader reflection information under Android",
+	ConfigRestartRequired = true))
+	bool bStripShaderReflection;
+
 	// If selected, the checked architectures will be split into separate .apk files [CURRENTLY FOR FULL SOURCE GAMES ONLY]
 	// @todo android fat binary: Currently, there isn't much utility in merging multiple .so's into a single .apk except for debugging,
 	// but we can't properly handle multiple GPU architectures in a single .apk, so we are disabling the feature for now
diff --git a/Engine/Source/Runtime/Engine/Private/ShaderCompiler/ShaderCompiler.cpp b/Engine/Source/Runtime/Engine/Private/ShaderCompiler/ShaderCompiler.cpp
--- a/Engine/Source/Runtime/Engine/Private/ShaderCompiler/ShaderCompiler.cpp
+++ b/Engine/Source/Runtime/Engine/Private/ShaderCompiler/ShaderCompiler.cpp
@@ -5467,6 +5467,16 @@ void GlobalBeginCompileShader(
 		const bool bIsMobileMultiViewCVar = CVarMobileMultiView && CVarMobileHDR ?
 			(CVarMobileMultiView->GetValueOnAnyThread() != 0 && CVarMobileHDR->GetValueOnAnyThread() == 0) : false;
 		const bool bIsODSCapture = CVarODSCapture && (CVarODSCapture->GetValueOnAnyThread() != 0);
+		
+		if (ShaderPlatform == SP_VULKAN_ES3_1_ANDROID || ShaderPlatform == SP_VULKAN_SM5_ANDROID)
+		{
+			bool bIsStripReflect = true;
+			GConfig->GetBool(TEXT("/Script/AndroidRuntimeSettings.AndroidRuntimeSettings"), TEXT("bStripShaderReflection"), bIsStripReflect, GEngineIni);
+			if (!bIsStripReflect)
+			{
+				Input.Environment.SetDefine(TEXT("STRIP_REFLECT_ANDROID"), false);
+			}
+		}
 
 		bool bIsInstancedStereo = !bUsingMobileRenderer && bIsInstancedStereoCVar && RHISupportsInstancedStereo(ShaderPlatform);
 		bool bIsMobileMultiview = bUsingMobileRenderer && bIsMobileMultiViewCVar;
diff --git a/Engine/Source/Runtime/RenderCore/Private/Shader.cpp b/Engine/Source/Runtime/RenderCore/Private/Shader.cpp
--- a/Engine/Source/Runtime/RenderCore/Private/Shader.cpp
+++ b/Engine/Source/Runtime/RenderCore/Private/Shader.cpp
@@ -1663,6 +1663,16 @@ void ShaderMapAppendKeyString(EShaderPlatform Platform, FString& KeyString)
 		}
 	}
 
+	if (Platform == SP_VULKAN_ES3_1_ANDROID || Platform == SP_VULKAN_SM5_ANDROID)
+	{
+		bool bStripReflect = true;
+		GConfig->GetBool(TEXT("/Script/AndroidRuntimeSettings.AndroidRuntimeSettings"), TEXT("bStripShaderReflection"), bStripReflect, GEngineIni);
+		if (!bStripReflect)
+		{
+			KeyString += TEXT("_NoStripReflect");
+		}
+	}
+
 	// Is DXC shader compiler enabled for this platform?
 	KeyString += (IsDxcEnabledForPlatform(Platform) ? TEXT("_DXC1") : TEXT("_DXC0"));
 
