UE5 - Remove non-ShaderConductor Vulkan shader compilation path as it's not used anymore

diff --git a/Engine/Source/Developer/VulkanShaderFormat/Private/VulkanGlslangBridge.cpp b/Engine/Source/Developer/VulkanShaderFormat/Private/VulkanGlslangBridge.cpp
--- a/Engine/Source/Developer/VulkanShaderFormat/Private/VulkanGlslangBridge.cpp
+++ b/Engine/Source/Developer/VulkanShaderFormat/Private/VulkanGlslangBridge.cpp
@@ -4,167 +4,8 @@
 #include "VulkanShaderFormat.h"
 #include "ShaderCore.h"
 #include "ShaderCompilerCore.h"
-
-#if defined(_MSC_VER) && _MSC_VER == 1800
-	#pragma warning(push)
-	#pragma warning(disable:4510) // warning C4510: 'glslang::TShader::Includer::IncludeResult' : default constructor could not be generated
-	#pragma warning(disable:4610) // error C4610: struct 'glslang::TShader::Includer::IncludeResult' can never be instantiated - user defined constructor required
-#endif // _MSC_VER == 1800
-#if defined(_MSC_VER) && _MSC_VER >= 1800
-	#pragma warning(push)
-	#pragma warning(disable: 6031) // warning C6031: Return value ignored: 'snprintf'.
-	#pragma warning(disable: 6340) // warning C6340: Mismatch on sign: 'const unsigned int' passed as _Param_(4) when some signed type is required in call to 'snprintf'.
-#endif
-
-THIRD_PARTY_INCLUDES_START
-	#include "glslang/Include/ShHandle.h"
-	#include "glslang/Include/revision.h"
-	#include "glslang/Public/ShaderLang.h"
-	#include "SPIRV/GlslangToSpv.h"
-	#include "SPIRV/GLSL.std.450.h"
-	#include "SPIRV/doc.h"
-	#include "SPIRV/disassemble.h"
-THIRD_PARTY_INCLUDES_END
-
-#if defined(_MSC_VER) && _MSC_VER >= 1800
-	#pragma warning(pop)
-#endif
-#if defined(_MSC_VER) && _MSC_VER == 1800
-	#pragma warning(pop)
-#endif // _MSC_VER == 1800
-
 #include "hlslcc.h"
 
-#include <fstream>
-
-struct FSpirvResources : public TBuiltInResource
-{
-	//#todo-rco: Make this based off platform (eg Android, PC, etc)
-	FSpirvResources()
-	{
-		maxLights = 32;
-		maxClipPlanes = 6;
-		maxTextureUnits = 32;
-		maxTextureCoords = 32;
-		maxVertexAttribs = 64;
-		maxVertexUniformComponents = 4096;
-		maxVaryingFloats = 64;
-		maxVertexTextureImageUnits = 32;
-		maxCombinedTextureImageUnits = 80;
-		maxTextureImageUnits = 32;
-		maxFragmentUniformComponents = 4096;
-		maxDrawBuffers = 32;
-		maxVertexUniformVectors = 128;
-		maxVaryingVectors = 8;
-		maxFragmentUniformVectors = 16;
-		maxVertexOutputVectors = 16;
-		maxFragmentInputVectors = 15;
-		minProgramTexelOffset = -8;
-		maxProgramTexelOffset = 7;
-		maxClipDistances = 8;
-		maxComputeWorkGroupCountX = 65535;
-		maxComputeWorkGroupCountY = 65535;
-		maxComputeWorkGroupCountZ = 65535;
-		maxComputeWorkGroupSizeX = 1024;
-		maxComputeWorkGroupSizeY = 1024;
-		maxComputeWorkGroupSizeZ = 64;
-		maxComputeUniformComponents = 1024;
-		maxComputeTextureImageUnits = 16;
-		maxComputeImageUniforms = 8;
-		maxComputeAtomicCounters = 8;
-		maxComputeAtomicCounterBuffers = 1;
-		maxVaryingComponents = 60;
-		maxVertexOutputComponents = 64;
-		maxGeometryInputComponents = 64;
-		maxGeometryOutputComponents = 128;
-		maxFragmentInputComponents = 128;
-		maxImageUnits = 8;
-		maxCombinedImageUnitsAndFragmentOutputs = 8;
-		maxCombinedShaderOutputResources = 8;
-		maxImageSamples = 0;
-		maxVertexImageUniforms = 0;
-		maxTessControlImageUniforms = 0;
-		maxTessEvaluationImageUniforms = 0;
-		maxGeometryImageUniforms = 0;
-		maxFragmentImageUniforms = 8;
-		maxCombinedImageUniforms = 8;
-		maxGeometryTextureImageUnits = 16;
-		maxGeometryOutputVertices = 256;
-		maxGeometryTotalOutputComponents = 1024;
-		maxGeometryUniformComponents = 1024;
-		maxGeometryVaryingComponents = 64;
-		maxTessControlInputComponents = 128;
-		maxTessControlOutputComponents = 128;
-		maxTessControlTextureImageUnits = 16;
-		maxTessControlUniformComponents = 1024;
-		maxTessControlTotalOutputComponents = 4096;
-		maxTessEvaluationInputComponents = 128;
-		maxTessEvaluationOutputComponents = 128;
-		maxTessEvaluationTextureImageUnits = 16;
-		maxTessEvaluationUniformComponents = 1024;
-		maxTessPatchComponents = 120;
-		maxPatchVertices = 32;
-		maxTessGenLevel = 64;
-		maxViewports = 16;
-		maxVertexAtomicCounters = 0;
-		maxTessControlAtomicCounters = 0;
-		maxTessEvaluationAtomicCounters = 0;
-		maxGeometryAtomicCounters = 0;
-		maxFragmentAtomicCounters = 8;
-		maxCombinedAtomicCounters = 8;
-		maxAtomicCounterBindings = 1;
-		maxVertexAtomicCounterBuffers = 0;
-		maxTessControlAtomicCounterBuffers = 0;
-		maxTessEvaluationAtomicCounterBuffers = 0;
-		maxGeometryAtomicCounterBuffers = 0;
-		maxFragmentAtomicCounterBuffers = 1;
-		maxCombinedAtomicCounterBuffers = 1;
-		maxAtomicCounterBufferSize = 16384;
-		maxTransformFeedbackBuffers = 4;
-		maxTransformFeedbackInterleavedComponents = 64;
-		maxCullDistances = 8;
-		maxCombinedClipAndCullDistances = 8;
-		maxSamples = 4;
-		limits.nonInductiveForLoops = 1;
-		limits.whileLoops = 1;
-		limits.doWhileLoops = 1;
-		limits.generalUniformIndexing = 1;
-		limits.generalAttributeMatrixVectorIndexing = 1;
-		limits.generalVaryingIndexing = 1;
-		limits.generalSamplerIndexing = 1;
-		limits.generalVariableIndexing = 1;
-		limits.generalConstantMatrixVectorIndexing = 1;
-
-
-		// One time init
-		glslang::InitializeProcess();
-	}
-
-	~FSpirvResources()
-	{
-		glslang::FinalizeProcess();
-	}
-};
-
-static FSpirvResources GSpirvResources;
-
-
-static EShLanguage GetStage(EHlslShaderFrequency Frequency)
-{
-	switch (Frequency)
-	{
-	case HSF_VertexShader: return EShLangVertex;
-	case HSF_PixelShader: return EShLangFragment;
-	case HSF_GeometryShader: return EShLangGeometry;
-	case HSF_ComputeShader: return EShLangCompute;
-	case HSF_HullShader: return EShLangTessControl;
-	case HSF_DomainShader: return EShLangTessEvaluation;
-	default: break;
-	}
-
-	return EShLangCount;
-}
-
 void PatchSpirvReflectionEntries(FVulkanSpirv& Spirv)
 {
 	TMap<uint32, FString> Names;
@@ -292,139 +133,3 @@ void PatchSpirvReflectionEntries(FVulkanSpirv& Spirv)
 		}
 	}
 }
-
-bool GenerateSpirv(const ANSICHAR* Source, FCompilerInfo& CompilerInfo, FString& OutErrors, const FString& DumpDebugInfoPath, FVulkanSpirv& OutSpirv)
-{
-	glslang::TProgram* Program = new glslang::TProgram;
-
-	EShLanguage Stage = GetStage(CompilerInfo.Frequency);
-	glslang::TShader* Shader = new glslang::TShader(Stage);
-
-	// Skip to #version
-	const char* GlslSourceSkipHeader = strstr(Source, "#version");
-	Shader->setStrings(&GlslSourceSkipHeader, 1);
-
-	Shader->setEntryPoint("main_00000000_00000000");
-	Shader->setSourceEntryPoint("main_00000000_00000000");
-
-	auto DoGenerate = [&]()
-	{
-		const int DefaultVersion = 100;// Options & EOptionDefaultDesktop ? 110 : 100;
-
-		EShMessages Messages = (EShMessages)(EShMsgDefault | EShMsgSpvRules | EShMsgVulkanRules);
-		if (!Shader->parse(&GSpirvResources, DefaultVersion, false, Messages))
-		{
-			OutErrors += ANSI_TO_TCHAR(Shader->getInfoLog());
-			return false;
-		}
-
-		Program->addShader(Shader);
-
-		if (!Program->link(Messages))
-		{
-			OutErrors += ANSI_TO_TCHAR(Program->getInfoLog());
-			return false;
-		}
-
-		if (!Program->getIntermediate(Stage))
-		{
-			OutErrors += ANSI_TO_TCHAR(Program->getInfoLog());
-			return false;
-		}
-
-		if (!Program->buildReflection())
-		{
-			OutErrors += ANSI_TO_TCHAR(Program->getInfoLog());
-			return false;
-		}
-
-		// Generate SPIRV
-		std::vector<unsigned int> Spirv;
-		glslang::GlslangToSpv(*Program->getIntermediate((EShLanguage)Stage), Spirv);
-
-		uint32 SizeInWords = Spirv.size();
-		OutSpirv.Data.AddZeroed(SizeInWords);
-		FMemory::Memcpy(OutSpirv.Data.GetData(), &Spirv[0], SizeInWords * sizeof(uint32));
-
-		// Gather Reflection info
-		for (int32 Index = 0; Index < Program->getNumLiveUniformVariables(); ++Index)
-		{
-			const char* AnsiName = Program->getUniformName(Index);
-			FString Name = AnsiName ? ANSI_TO_TCHAR(AnsiName) : TEXT("NULL");
-			const auto* Type = Program->getAttributeTType(Index);
-			int32 Binding = -1;
-			if (Type && Type->getQualifier().hasBinding())
-			{
-				Binding = Type->getQualifier().layoutBinding;
-			}
-			else
-			{
-				Binding = Program->getUniformBinding(Index);
-			}
-			FVulkanSpirv::FEntry Entry(Name, Binding);
-			OutSpirv.ReflectionInfo.Add(Entry);
-		}
-
-		for (int32 Index = 0; Index < Program->getNumLiveUniformBlocks(); ++Index)
-		{
-			const char* AnsiName = Program->getUniformBlockName(Index);
-			FString Name = AnsiName ? ANSI_TO_TCHAR(AnsiName) : TEXT("NULL");
-			const auto* Type = Program->getUniformBlockTType(Index);
-			int32 Binding = -1;
-			if (Type && Type->getQualifier().hasBinding())
-			{
-				Binding = Type->getQualifier().layoutBinding;
-			}
-			else
-			{
-				Binding = Program->getUniformBinding(Index);
-			}
-			FVulkanSpirv::FEntry Entry(Name, Binding);
-			OutSpirv.ReflectionInfo.Add(Entry);
-		}
-
-		PatchSpirvReflectionEntries(OutSpirv);
-		OutSpirv.EntryPointName = PatchSpirvEntryPointWithCRC(OutSpirv, OutSpirv.CRC);
-
-		// Copy back to original spirv data as it is used for dumping information
-		FMemory::Memcpy(&Spirv[0], OutSpirv.Data.GetData(), SizeInWords * sizeof(uint32));
-
-		// Dump debug info/files
-		if (CompilerInfo.bDebugDump)
-		{
-			// Binary SpirV
-			FString SpirvFile = DumpDebugInfoPath / (TEXT("Output.spv"));
-			glslang::OutputSpvBin(Spirv, TCHAR_TO_ANSI(*SpirvFile));
-
-			// Text spirv
-			FString SpirvTextFile = DumpDebugInfoPath / (TEXT("Output.spvasm"));
-			std::ofstream File;
-			File.open(TCHAR_TO_ANSI(*SpirvTextFile), std::fstream::out/* | std::fstream::binary*/);
-			if (File.is_open())
-			{
-				spv::Parameterize();
-				spv::Disassemble(File, Spirv);
-				File.close();
-			}
-		}
-		else if (CompilerInfo.Input.bSkipPreprocessedCache)
-		{
-			spv::Parameterize();
-			spv::Disassemble(std::cout, Spirv);
-		}
-
-		if (CompilerInfo.Input.bSkipPreprocessedCache)
-		{
-			Program->dumpReflection();
-		}
-
-		return true;
-	};
-
-	bool bResult = DoGenerate();
-
-	delete Program;
-	delete Shader;
-
-	return bResult;
-}
diff --git a/Engine/Source/Developer/VulkanShaderFormat/Private/VulkanShaderCompiler.cpp b/Engine/Source/Developer/VulkanShaderFormat/Private/VulkanShaderCompiler.cpp
--- a/Engine/Source/Developer/VulkanShaderFormat/Private/VulkanShaderCompiler.cpp
+++ b/Engine/Source/Developer/VulkanShaderFormat/Private/VulkanShaderCompiler.cpp
@@ -1409,68 +1409,6 @@ FCompilerInfo::FCompilerInfo(const FShaderCompilerInput& InInput, const FString&
 	BaseSourceFilename = Input.GetSourceFilename();
 }
 
-/**
- * Compile a shader using the internal shader compiling library
- */
-static bool CompileUsingInternal(FCompilerInfo& CompilerInfo, const FVulkanBindingTable& BindingTable, const TArray<ANSICHAR>& GlslSource, FShaderCompilerOutput& Output)
-{
-	FString Errors;
-	FVulkanSpirv Spirv;
-	const ANSICHAR* Main = GlslSource.GetData();
-	Main = FCStringAnsi::Strstr(Main, "void main_");
-	check(Main);
-	auto GetNumEOLs = [](const ANSICHAR* Ptr)
-	{
-		uint32 NumLines = 0;
-		while (*Ptr)
-		{
-			if (*Ptr == '\n')
-			{
-				++NumLines;
-			}
-			++Ptr;
-		}
-
-		return NumLines;
-	};
-	uint32 NumLines = GetNumEOLs(Main);
-	if (GenerateSpirv(GlslSource.GetData(), CompilerInfo, Errors, CompilerInfo.Input.DumpDebugInfoPath, Spirv))
-	{
-		FString DebugName = CompilerInfo.Input.DumpDebugInfoPath.Right(CompilerInfo.Input.DumpDebugInfoPath.Len() - CompilerInfo.Input.DumpDebugInfoRootPath.Len());
-
-		Output.Target = CompilerInfo.Input.Target;
-		BuildShaderOutput(Output, CompilerInfo.Input,
-			GlslSource.GetData(), GlslSource.Num(),
-			BindingTable, NumLines, Spirv, DebugName, false);
-
-		if (CompilerInfo.bDebugDump)
-		{
-			FString InfoFile = CompilerInfo.Input.DumpDebugInfoPath / TEXT("Info.txt");
-			FArchive* FileWriter = IFileManager::Get().CreateFileWriter(*InfoFile);
-			if (FileWriter)
-			{
-				FString OutputString = FString::Printf(TEXT("main_%0.8x_%0.8x\n"), Spirv.GetByteSize(), Spirv.CRC);
-				auto AnsiOutputString = StringCast<ANSICHAR>(*OutputString);
-				FileWriter->Serialize((ANSICHAR*)AnsiOutputString.Get(), AnsiOutputString.Length());
-				FileWriter->Close();
-			}
-			delete FileWriter;
-		}
-		return true;
-	}
-	else
-	{
-		if (Errors.Len() > 0)
-		{
-			FShaderCompilerError* Error = new(Output.Errors) FShaderCompilerError();
-			Error->StrippedErrorMessage = Errors;
-		}
-
-		return false;
-	}
-}
-
-
 #if PLATFORM_MAC || PLATFORM_WINDOWS || PLATFORM_LINUX
 
 static void GatherSpirvReflectionBindings(
diff --git a/Engine/Source/Developer/VulkanShaderFormat/Private/VulkanShaderFormat.h b/Engine/Source/Developer/VulkanShaderFormat/Private/VulkanShaderFormat.h
--- a/Engine/Source/Developer/VulkanShaderFormat/Private/VulkanShaderFormat.h
+++ b/Engine/Source/Developer/VulkanShaderFormat/Private/VulkanShaderFormat.h
@@ -136,6 +136,3 @@ struct FVulkanSpirv : FSpirv
 
 // Updates all reflection entries in the specified SPIR-V module.
 extern void PatchSpirvReflectionEntries(FVulkanSpirv& OutSpirv);
-
-// Generates SPIR-V out of the specified GLSL source code.
-extern bool GenerateSpirv(const ANSICHAR* Source, FCompilerInfo& CompilerInfo, FString& OutErrors, const FString& DumpDebugInfoPath, FVulkanSpirv& OutSpirv);
diff --git a/Engine/Source/Developer/VulkanShaderFormat/VulkanShaderFormat.Build.cs b/Engine/Source/Developer/VulkanShaderFormat/VulkanShaderFormat.Build.cs
--- a/Engine/Source/Developer/VulkanShaderFormat/VulkanShaderFormat.Build.cs
+++ b/Engine/Source/Developer/VulkanShaderFormat/VulkanShaderFormat.Build.cs
@@ -22,8 +22,6 @@ public class VulkanShaderFormat : ModuleRules
 			}
 			);
 
-		AddEngineThirdPartyPrivateStaticDependencies(Target, "GlsLang");
-
 		if (Target.Platform == UnrealTargetPlatform.Mac || Target.Platform == UnrealTargetPlatform.Win64 || Target.Platform == UnrealTargetPlatform.Linux)
 		{
 			AddEngineThirdPartyPrivateStaticDependencies(Target, "SPIRVReflect");
