Handle TextureExternal for vulkan in our preprocessor, not DXC

diff --git a/Engine/Shaders/Private/Common.ush b/Engine/Shaders/Private/Common.ush
--- a/Engine/Shaders/Private/Common.ush
+++ b/Engine/Shaders/Private/Common.ush
@@ -122,7 +122,7 @@ const static float3 Max111110BitsFloat3 = float3(Max11BitsFloat, Max11BitsFloat,
 #define SetSceneColorLayoutToFloat4(dest,value) dest.rgb = (value).rgb
 #endif
 
-#define SUPPORTS_TEXTURE_EXTERNAL		(COMPILER_GLSL_ES3_1 || COMPILER_VULKAN)
+#define SUPPORTS_TEXTURE_EXTERNAL		(COMPILER_GLSL_ES3_1)
 
 #if !SUPPORTS_TEXTURE_EXTERNAL
 #define TextureExternal Texture2D
diff --git a/Engine/Source/Developer/VulkanShaderFormat/Private/VulkanShaderCompiler.cpp b/Engine/Source/Developer/VulkanShaderFormat/Private/VulkanShaderCompiler.cpp
--- a/Engine/Source/Developer/VulkanShaderFormat/Private/VulkanShaderCompiler.cpp
+++ b/Engine/Source/Developer/VulkanShaderFormat/Private/VulkanShaderCompiler.cpp
@@ -2055,7 +2055,6 @@ static bool CompileWithShaderConductor(
 
 	// Inject additional macro definitions to circumvent missing features: external textures
 	FShaderCompilerDefinitions AdditionalDefines;
-	AdditionalDefines.SetDefine(TEXT("TextureExternal"), TEXT("Texture2D"));
 
 	if (bDebugDump)
 	{
