Adds a CompileDXC.bat to vulkan's ShaderDebugInfo

This allows to iterate to compile the USF file to SPIRV and translate to GLSL to verify the sanity of the SPIRV assembly.

diff --git a/Engine/Source/Developer/VulkanShaderFormat/Private/VulkanShaderCompiler.cpp b/Engine/Source/Developer/VulkanShaderFormat/Private/VulkanShaderCompiler.cpp
--- a/Engine/Source/Developer/VulkanShaderFormat/Private/VulkanShaderCompiler.cpp
+++ b/Engine/Source/Developer/VulkanShaderFormat/Private/VulkanShaderCompiler.cpp
@@ -2036,6 +2036,115 @@ static void Patch64bitSamplers(FVulkanSpirv& Spirv)
 	}
 }
 
+static FString VulkanGetShaderProfileDXC(const FCompilerInfo& CompilerInfo, const CrossCompiler::FShaderConductorOptions Options)
+{
+	const TCHAR* ShaderProfile = TEXT("unknown");
+	switch (CompilerInfo.Input.Target.Frequency)
+	{
+	case SF_Vertex:
+		ShaderProfile = TEXT("vs");
+		break;
+
+	case SF_Pixel:
+		ShaderProfile = TEXT("ps");
+		break;
+
+	case SF_Geometry:
+		ShaderProfile = TEXT("gs");
+		break;
+
+	case SF_Compute:
+		ShaderProfile = TEXT("cs");
+		break;
+
+	case SF_RayGen:
+	case SF_RayMiss:
+	case SF_RayHitGroup:
+	case SF_RayCallable:
+		return TEXT("lib_6_3");
+	}
+
+	return FString::Printf(TEXT("%s_%d_%d"), ShaderProfile, Options.ShaderModel.Major, Options.ShaderModel.Minor); 
+}
+
+static void VulkanCreateDXCCompileBatchFiles(
+	const FString& EntryPointName,
+	EShaderFrequency Frequency,
+	const FCompilerInfo& CompilerInfo,
+	const CrossCompiler::FShaderConductorOptions Options)
+{
+
+	FString USFFilename = CompilerInfo.Input.GetSourceFilename();
+	FString SPVFilename = FPaths::GetBaseFilename(USFFilename) + TEXT(".DXC.spv");
+	FString GLSLFilename = FPaths::GetBaseFilename(USFFilename) + TEXT(".SPV.glsl");
+
+	FString DxcPath = FPaths::ConvertRelativePathToFull(FPaths::EngineDir());
+
+	DxcPath = FPaths::Combine(DxcPath, TEXT("Binaries/ThirdParty/ShaderConductor/Win64"));
+	FPaths::MakePlatformFilename(DxcPath);
+
+	FString DxcFilename = FPaths::Combine(DxcPath, TEXT("dxc.exe"));
+	FPaths::MakePlatformFilename(DxcFilename);
+
+	const TCHAR* VulkanVersion = TEXT("vulkanUNKNOWN");
+	if (Options.TargetEnvironment == CrossCompiler::FShaderConductorOptions::ETargetEnvironment::Vulkan_1_0)
+	{
+		VulkanVersion = TEXT("vulkan1.0");
+	}
+	else if (Options.TargetEnvironment == CrossCompiler::FShaderConductorOptions::ETargetEnvironment::Vulkan_1_1)
+	{
+		VulkanVersion = TEXT("vulkan1.1");
+	}
+	else if (Options.TargetEnvironment == CrossCompiler::FShaderConductorOptions::ETargetEnvironment::Vulkan_1_2)
+	{
+		VulkanVersion = TEXT("vulkan1.2");
+	}
+	else
+	{
+		ensure(false);
+	}
+
+	FString ShaderProfile = VulkanGetShaderProfileDXC(CompilerInfo, Options);
+
+	// CompileDXC.bat
+	{
+		FString BatchFileContents =  FString::Printf(
+			TEXT(
+				"@ECHO OFF\n"
+				"SET DXC=\"%s\"\n"
+				"SET SPIRVCROSS=\"spirv-cross.exe\"\n"
+				"IF NOT EXIST %%DXC%% (\n"
+				"\tECHO Couldn't find dxc.exe under \"%s\"\n"
+				"\tGOTO :END\n"
+				")\n"
+				"ECHO Compiling with DXC...\n"
+				"%%DXC%% -HV %d -T %s -E %s -spirv -fspv-target-env=%s -Fo %s %s\n"
+				"WHERE %%SPIRVCROSS%%\n"
+				"IF %%ERRORLEVEL%% NEQ 0 (\n"
+				"\tECHO spirv-cross.exe not found in Path environment variable, please build it from source [REDACTED_PROFILE_URL]/SPIRV-Cross\n"
+				"\tGOTO :END\n"
+				")\n"
+				"ECHO Translating SPIRV back to glsl...\n"
+				"%%SPIRVCROSS%% --vulkan-semantics --output %s %s\n"
+				":END\n"
+				"PAUSE\n"
+			),
+			*DxcFilename,
+			*DxcPath,
+			Options.HlslVersion,
+			*ShaderProfile,
+			*EntryPointName,
+			VulkanVersion,
+			*SPVFilename,
+			*USFFilename,
+			*GLSLFilename,
+			*SPVFilename
+		);
+
+		FFileHelper::SaveStringToFile(BatchFileContents, *(CompilerInfo.Input.DumpDebugInfoPath / TEXT("CompileDXC.bat")));
+	}
+}
+
 static bool CompileWithShaderConductor(
 	const FString&			PreprocessedShader,
 	const FString&			EntryPointName,
@@ -2056,11 +2165,6 @@ static bool CompileWithShaderConductor(
 	// Inject additional macro definitions to circumvent missing features: external textures
 	FShaderCompilerDefinitions AdditionalDefines;
 
-	if (bDebugDump)
-	{
-		DumpDebugUSF(Input, PreprocessedShader, CompilerInfo.CCFlags);
-	}
-
 	// Load shader source into compiler context
 	CompilerContext.LoadSource(PreprocessedShader, Input.VirtualSourceFilePath, EntryPointName, Frequency, &AdditionalDefines);
 
@@ -2084,6 +2188,17 @@ static bool CompileWithShaderConductor(
 		Options.TargetEnvironment = GetMinimumTargetEnvironment(Input.Target.GetPlatform());
 	}
 
+	if (bDebugDump)
+	{
+		VulkanCreateDXCCompileBatchFiles(
+			EntryPointName,
+			Frequency,
+			CompilerInfo,
+			Options);
+
+		DumpDebugUSF(Input, PreprocessedShader, CompilerInfo.CCFlags);
+	}
+
 	if (bRewriteHlslSource)
 	{
 		// Rewrite HLSL source code to remove unused global resources and variables
