Move duplicated implementations of ParseRayTracingEntryPoint() to ShaderCommon module

diff --git a/Engine/Source/Developer/ShaderCompilerCommon/Private/ShaderCompilerCommon.cpp b/Engine/Source/Developer/ShaderCompilerCommon/Private/ShaderCompilerCommon.cpp
--- a/Engine/Source/Developer/ShaderCompilerCommon/Private/ShaderCompilerCommon.cpp
+++ b/Engine/Source/Developer/ShaderCompilerCommon/Private/ShaderCompilerCommon.cpp
@@ -625,6 +625,37 @@ bool UE::ShaderCompilerCommon::ValidatePackedResourceCounts(FShaderCompilerOutpu
 	return Output.bSucceeded;
 }
 
+void UE::ShaderCompilerCommon::ParseRayTracingEntryPoint(const FString& Input, FString& OutMain, FString& OutAnyHit, FString& OutIntersection)
+{
+	auto ParseEntry = [&Input](const TCHAR* Marker)
+	{
+		FString Result;
+		const int32 BeginIndex = Input.Find(Marker, ESearchCase::IgnoreCase, ESearchDir::FromStart);
+		if (BeginIndex != INDEX_NONE)
+		{
+			int32 EndIndex = Input.Find(TEXT(" "), ESearchCase::IgnoreCase, ESearchDir::FromStart, BeginIndex);
+			if (EndIndex == INDEX_NONE)
+			{
+				EndIndex = Input.Len() + 1;
+			}
+			const int32 MarkerLen = FCString::Strlen(Marker);
+			const int32 Count = EndIndex - BeginIndex;
+			Result = Input.Mid(BeginIndex + MarkerLen, Count - MarkerLen);
+		}
+		return Result;
+	};
+
+	OutMain = ParseEntry(TEXT("closesthit="));
+	OutAnyHit = ParseEntry(TEXT("anyhit="));
+	OutIntersection = ParseEntry(TEXT("intersection="));
+
+	// If complex hit group entry is not specified, assume a single verbatim entry point
+	if (OutMain.IsEmpty() && OutAnyHit.IsEmpty() && OutIntersection.IsEmpty())
+	{
+		OutMain = Input;
+	}
+}
+
 void HandleReflectedGlobalConstantBufferMember(
 	const FString& InMemberName,
 	uint32 ConstantBufferIndex,
diff --git a/Engine/Source/Developer/ShaderCompilerCommon/Public/ShaderCompilerCommon.h b/Engine/Source/Developer/ShaderCompilerCommon/Public/ShaderCompilerCommon.h
--- a/Engine/Source/Developer/ShaderCompilerCommon/Public/ShaderCompilerCommon.h
+++ b/Engine/Source/Developer/ShaderCompilerCommon/Public/ShaderCompilerCommon.h
@@ -67,6 +67,18 @@ namespace UE::ShaderCompilerCommon
 	extern SHADERCOMPILERCOMMON_API bool                 RemoveBindlessParameterPrefix(FString& InName);
 
 	extern SHADERCOMPILERCOMMON_API bool                 ValidatePackedResourceCounts(FShaderCompilerOutput& Output, const FShaderCodePackedResourceCounts& PackedResourceCounts);
+
+	/*
+	* Parses ray tracing shader entry point specification string in one of the following formats:
+	* 1) Verbatim single entry point name, e.g. "MainRGS"
+	* 2) Complex entry point for ray tracing hit group shaders:
+	*      a) "closesthit=MainCHS"
+	*      b) "closesthit=MainCHS anyhit=MainAHS"
+	*      c) "closesthit=MainCHS anyhit=MainAHS intersection=MainIS"
+	*      d) "closesthit=MainCHS intersection=MainIS"
+	*    NOTE: closesthit attribute must always be provided for complex hit group entry points
+	*/
+	extern SHADERCOMPILERCOMMON_API void ParseRayTracingEntryPoint(const FString& Input, FString& OutMain, FString& OutAnyHit, FString& OutIntersection);
 }
 
 extern SHADERCOMPILERCOMMON_API void HandleReflectedGlobalConstantBufferMember(
diff --git a/Engine/Source/Developer/VulkanShaderFormat/Private/VulkanShaderCompiler.cpp b/Engine/Source/Developer/VulkanShaderFormat/Private/VulkanShaderCompiler.cpp
--- a/Engine/Source/Developer/VulkanShaderFormat/Private/VulkanShaderCompiler.cpp
+++ b/Engine/Source/Developer/VulkanShaderFormat/Private/VulkanShaderCompiler.cpp
@@ -200,45 +200,6 @@ static uint16 GetCombinedSamplerStateAlias(const FString& ParameterName,
 	return UINT16_MAX;
 }
 
-// Parses ray tracing shader entry point specification string in one of the following formats:
-// 1) Verbatim single entry point name, e.g. "MainRGS"
-// 2) Complex entry point for ray tracing hit group shaders:
-//      a) "closesthit=MainCHS"
-//      b) "closesthit=MainCHS anyhit=MainAHS"
-//      c) "closesthit=MainCHS anyhit=MainAHS intersection=MainIS"
-//      d) "closesthit=MainCHS intersection=MainIS"
-//    NOTE: closesthit attribute must always be provided for complex hit group entry points
-static void ParseRayTracingEntryPoint(const FString& Input, FString& OutMain, FString& OutAnyHit, FString& OutIntersection)
-{
-	auto ParseEntry = [&Input](const TCHAR* Marker)
-	{
-		FString Result;
-		const int32 BeginIndex = Input.Find(Marker, ESearchCase::IgnoreCase, ESearchDir::FromStart);
-		if (BeginIndex != INDEX_NONE)
-		{
-			int32 EndIndex = Input.Find(TEXT(" "), ESearchCase::IgnoreCase, ESearchDir::FromStart, BeginIndex);
-			if (EndIndex == INDEX_NONE)
-			{
-				EndIndex = Input.Len() + 1;
-			}
-			const int32 MarkerLen = FCString::Strlen(Marker);
-			const int32 Count = EndIndex - BeginIndex;
-			Result = Input.Mid(BeginIndex + MarkerLen, Count - MarkerLen);
-		}
-		return Result;
-	};
-
-	OutMain = ParseEntry(TEXT("closesthit="));
-	OutAnyHit = ParseEntry(TEXT("anyhit="));
-	OutIntersection = ParseEntry(TEXT("intersection="));
-
-	// If complex hit group entry is not specified, assume a single verbatim entry point
-	if (OutMain.IsEmpty() && OutAnyHit.IsEmpty() && OutIntersection.IsEmpty())
-	{
-		OutMain = Input;
-	}
-}
-
 struct FPatchType
 {
 	int32	HeaderGlobalIndex;
@@ -1512,7 +1473,7 @@ static bool BuildShaderOutputFromSpirv(
 	{
 		// For now only use the primary entry point for hit groups until we decide how to best support hit group library shaders.
 		FString OutMain, OutAnyHit, OutIntersection;
-		ParseRayTracingEntryPoint(EntryPointName, OutMain, OutAnyHit, OutIntersection);
+		UE::ShaderCompilerCommon::ParseRayTracingEntryPoint(EntryPointName, OutMain, OutAnyHit, OutIntersection);
 
 		for (uint32 i = 0; i < Reflection.GetEntryPointCount(); ++i)
 		{
diff --git a/Engine/Source/Developer/Windows/ShaderFormatD3D/Private/D3DShaderCompilerDXC.cpp b/Engine/Source/Developer/Windows/ShaderFormatD3D/Private/D3DShaderCompilerDXC.cpp
--- a/Engine/Source/Developer/Windows/ShaderFormatD3D/Private/D3DShaderCompilerDXC.cpp
+++ b/Engine/Source/Developer/Windows/ShaderFormatD3D/Private/D3DShaderCompilerDXC.cpp
@@ -413,42 +413,6 @@ inline bool IsCompatibleBinding(const D3D12_SHADER_INPUT_BIND_DESC& BindDesc, ui
 	return bIsCompatibleBinding;
 }
 
-// Parses ray tracing shader entry point specification string in one of the following formats:
-// 1) Verbatim single entry point name, e.g. "MainRGS"
-// 2) Complex entry point for ray tracing hit group shaders:
-//      a) "closesthit=MainCHS"
-//      b) "closesthit=MainCHS anyhit=MainAHS"
-//      c) "closesthit=MainCHS anyhit=MainAHS intersection=MainIS"
-//      d) "closesthit=MainCHS intersection=MainIS"
-//    NOTE: closesthit attribute must always be provided for complex hit group entry points
-static void ParseRayTracingEntryPoint(const FString& Input, FString& OutMain, FString& OutAnyHit, FString& OutIntersection)
-{
-	auto ParseEntry = [&Input](const TCHAR* Marker)
-	{
-		FString Result;
-		int32 BeginIndex = Input.Find(Marker, ESearchCase::IgnoreCase, ESearchDir::FromStart);
-		if (BeginIndex != INDEX_NONE)
-		{
-			int32 EndIndex = Input.Find(TEXT(" "), ESearchCase::IgnoreCase, ESearchDir::FromStart, BeginIndex);
-			if (EndIndex == INDEX_NONE) EndIndex = Input.Len() + 1;
-			int32 MarkerLen = FCString::Strlen(Marker);
-			int32 Count = EndIndex - BeginIndex;
-			Result = Input.Mid(BeginIndex + MarkerLen, Count - MarkerLen);
-		}
-		return Result;
-	};
-
-	OutMain = ParseEntry(TEXT("closesthit="));
-	OutAnyHit = ParseEntry(TEXT("anyhit="));
-	OutIntersection = ParseEntry(TEXT("intersection="));
-
-	// If complex hit group entry is not specified, assume a single verbatim entry point
-	if (OutMain.IsEmpty() && OutAnyHit.IsEmpty() && OutIntersection.IsEmpty())
-	{
-		OutMain = Input;
-	}
-}
-
 static ShaderConductor::Compiler::ShaderModel ToDXCShaderModel(ELanguage Language)
 {
 	switch (Language)
@@ -625,7 +589,7 @@ bool CompileAndProcessD3DShaderDXC(FString& PreprocessedShaderSource,
 
 	if (bIsRayTracingShader)
 	{
-		ParseRayTracingEntryPoint(Input.EntryPointName, RayEntryPoint, RayAnyHitEntryPoint, RayIntersectionEntryPoint);
+		UE::ShaderCompilerCommon::ParseRayTracingEntryPoint(Input.EntryPointName, RayEntryPoint, RayAnyHitEntryPoint, RayIntersectionEntryPoint);
 
 		RayTracingExports = RayEntryPoint;
 
