- Checking the SF instead of the SP in the Vulkan Shader Compiler

diff --git a/Engine/Source/Developer/VulkanShaderFormat/Private/VulkanShaderCompiler.cpp b/Engine/Source/Developer/VulkanShaderFormat/Private/VulkanShaderCompiler.cpp
--- a/Engine/Source/Developer/VulkanShaderFormat/Private/VulkanShaderCompiler.cpp
+++ b/Engine/Source/Developer/VulkanShaderFormat/Private/VulkanShaderCompiler.cpp
@@ -9,6 +9,7 @@
 #include "HlslccHeaderWriter.h"
 #include "hlslcc.h"
 #include "SpirvReflectCommon.h"
+#include "RHIShaderFormatDefinitions.inl"
 
 #if PLATFORM_MAC
 // Horrible hack as we need the enum available but the Vulkan headers do not compile on Mac
@@ -38,25 +39,25 @@ enum VkDescriptorType {
 #include "Misc/FileHelper.h"
 #include "Misc/Paths.h"
 
-inline bool CanCompilePlatform(EShaderPlatform ShaderPlatform)
+inline bool IsVulkanShaderFormat(FName ShaderFormat)
 {
-	return ShaderPlatform == SP_VULKAN_PCES3_1
-		|| ShaderPlatform == SP_VULKAN_SM5
-		|| ShaderPlatform == SP_VULKAN_ES3_1_ANDROID
-		|| ShaderPlatform == SP_VULKAN_SM5_ANDROID;
+	return ShaderFormat == NAME_VULKAN_ES3_1_ANDROID
+		|| ShaderFormat == NAME_VULKAN_ES3_1
+		|| ShaderFormat == NAME_VULKAN_SM5
+		|| ShaderFormat == NAME_VULKAN_SM5_ANDROID;
 }
 
-inline bool IsAndroidShaderPlatform(EShaderPlatform ShaderPlatform)
+inline bool IsAndroidShaderFormat(FName ShaderFormat)
 {
-	return ShaderPlatform == SP_VULKAN_ES3_1_ANDROID
-		|| ShaderPlatform == SP_VULKAN_SM5_ANDROID;
+	return ShaderFormat == NAME_VULKAN_ES3_1_ANDROID
+		|| ShaderFormat == NAME_VULKAN_SM5_ANDROID;
 }
 
-inline bool SupportsOfflineCompiler(EShaderPlatform ShaderPlatform)
+inline bool SupportsOfflineCompiler(FName ShaderFormat)
 {
-	return ShaderPlatform == SP_VULKAN_PCES3_1
-		|| ShaderPlatform == SP_VULKAN_ES3_1_ANDROID
-		|| ShaderPlatform == SP_VULKAN_SM5_ANDROID;
+	return ShaderFormat == NAME_VULKAN_ES3_1_ANDROID
+		|| ShaderFormat == NAME_VULKAN_ES3_1
+		|| ShaderFormat == NAME_VULKAN_SM5_ANDROID;
 }
 
 inline CrossCompiler::FShaderConductorOptions::ETargetEnvironment GetMinimumTargetEnvironment(EShaderPlatform ShaderPlatform)
@@ -1328,7 +1329,7 @@ static void BuildShaderOutput(
 	}
 	if (ShaderInput.ExtraSettings.OfflineCompilerPath.Len() > 0)
 	{
-		if (SupportsOfflineCompiler(ShaderInput.Target.GetPlatform()))
+		if (SupportsOfflineCompiler(ShaderInput.ShaderFormat))
 		{
 			CompileOfflineMali(ShaderInput, ShaderOutput, (const ANSICHAR*)Spirv.GetByteData(), Spirv.GetByteSize(), true, Spirv.EntryPointName);
 		}
@@ -1861,7 +1862,7 @@ static bool BuildShaderOutputFromSpirv(
 	}
 
 	// For Android run an additional pass to patch spirv to be compatible across drivers
-	if (IsAndroidShaderPlatform(Input.Target.GetPlatform()))
+	if(IsAndroidShaderFormat(Input.ShaderFormat))
 	{
 		const char* OptArgs[] = { "--android-driver-patch" };
 		if (!CompilerContext.OptimizeSpirv(Spirv.Data, OptArgs, UE_ARRAY_COUNT(OptArgs)))
@@ -2244,13 +2245,13 @@ static bool CompileWithShaderConductor(
 
 void DoCompileVulkanShader(const FShaderCompilerInput& Input, FShaderCompilerOutput& Output, const class FString& WorkingDirectory, EVulkanShaderVersion Version)
 {
-	check(CanCompilePlatform(Input.Target.GetPlatform()));
+	check(IsVulkanShaderFormat(Input.ShaderFormat));
 
 	const bool bIsSM5 = (Version == EVulkanShaderVersion::SM5);
 	const bool bIsMobile = (Version == EVulkanShaderVersion::ES3_1 || Version == EVulkanShaderVersion::ES3_1_ANDROID);
 	bool bStripReflect = Input.IsRayTracingShader();
 	// By default we strip reflecion information for Android platform to avoid issues with older drivers
-	if (IsAndroidShaderPlatform(Input.Target.GetPlatform()))
+	if (IsAndroidShaderFormat(Input.ShaderFormat))
 	{
 		const FString* StripReflect_Android = Input.Environment.GetDefinitions().Find(TEXT("STRIP_REFLECT_ANDROID"));
 		bStripReflect = !(StripReflect_Android && *StripReflect_Android == TEXT("0"));
diff --git a/Engine/Source/Developer/VulkanShaderFormat/Private/VulkanShaderFormat.cpp b/Engine/Source/Developer/VulkanShaderFormat/Private/VulkanShaderFormat.cpp
--- a/Engine/Source/Developer/VulkanShaderFormat/Private/VulkanShaderFormat.cpp
+++ b/Engine/Source/Developer/VulkanShaderFormat/Private/VulkanShaderFormat.cpp
@@ -11,11 +11,7 @@
 #include "ShaderCompilerCore.h"
 #include "DXCWrapper.h"
 #include "ShaderConductorContext.h"
-
-static FName NAME_VULKAN_ES3_1_ANDROID(TEXT("SF_VULKAN_ES31_ANDROID"));
-static FName NAME_VULKAN_ES3_1(TEXT("SF_VULKAN_ES31"));
-static FName NAME_VULKAN_SM5(TEXT("SF_VULKAN_SM5"));
-static FName NAME_VULKAN_SM5_ANDROID(TEXT("SF_VULKAN_SM5_ANDROID"));
+#include "RHIShaderFormatDefinitions.inl"
 
 class FShaderFormatVulkan : public IShaderFormat
 {
