Bump vulkan shader format version to invalidate Android VK shaders

diff --git a/Engine/Source/Developer/VulkanShaderFormat/Private/VulkanShaderFormat.cpp b/Engine/Source/Developer/VulkanShaderFormat/Private/VulkanShaderFormat.cpp
--- a/Engine/Source/Developer/VulkanShaderFormat/Private/VulkanShaderFormat.cpp
+++ b/Engine/Source/Developer/VulkanShaderFormat/Private/VulkanShaderFormat.cpp
@@ -17,8 +17,8 @@ class FShaderFormatVulkan : public IShaderFormat
 {
 	enum 
 	{
-		UE_SHADER_VULKAN_ES3_1_VER	= 36,
-		UE_SHADER_VULKAN_SM5_VER 	= 36,
+		UE_SHADER_VULKAN_ES3_1_VER	= 37,
+		UE_SHADER_VULKAN_SM5_VER 	= 37,
 	};
 
 	int32 InternalGetVersion(FName Format) const
