VulkanRHI: Quick patch to add support for WAVESIZE in Vulkan.

diff --git a/Engine/Shaders/Public/Platform/Vulkan/VulkanCommon.ush b/Engine/Shaders/Public/Platform/Vulkan/VulkanCommon.ush
--- a/Engine/Shaders/Public/Platform/Vulkan/VulkanCommon.ush
+++ b/Engine/Shaders/Public/Platform/Vulkan/VulkanCommon.ush
@@ -5,7 +5,7 @@
 =============================================================================*/
 
 // Update this GUID to improve shader recompilation for Vulkan only shaders
-// GUID = 494E180E-0471-4013-8210-B52A0B05805F
+// GUID = 1084d248-f5cb-41ba-8ff7-3ee0752a4295
 
 #pragma once
 
@@ -50,3 +50,12 @@ void ImageInterlockedMaxUInt64(RWTexture2D<UlongType> Dest, uint2 Coord, UlongTy
 	SamplerComparisonState VULKAN_SAMPLER_HEAP(SamplerComparisonState)[] : register(s0, space1000);
 	COMPILER_DEFINE GetSamplerFromHeap(Type, Index) VULKAN_SAMPLER_HEAP(Type)[Index]
 #endif
+
+#if VULKAN_SUPPORTS_SUBGROUP_SIZE_CONTROL
+	#define COMPILER_SUPPORTS_WAVE_SIZE 1
+
+	// NOTE: Wave size in Vulkan is supported at shader stage in the pipeline creation.
+	// Create something unique we can read back from the shader, but that goes away on the rewrite
+	COMPILER_DEFINE VULKAN_WAVESIZE(N)
+	#define WAVESIZE(N) VULKAN_WAVESIZE(N)
+#endif
diff --git a/Engine/Source/Developer/VulkanShaderFormat/Private/VulkanShaderCompiler.cpp b/Engine/Source/Developer/VulkanShaderFormat/Private/VulkanShaderCompiler.cpp
--- a/Engine/Source/Developer/VulkanShaderFormat/Private/VulkanShaderCompiler.cpp
+++ b/Engine/Source/Developer/VulkanShaderFormat/Private/VulkanShaderCompiler.cpp
@@ -886,6 +886,7 @@ static void BuildShaderOutput(
 	int32						SourceLen,
 	const FVulkanBindingTable&	BindingTable,
 	uint32						NumLines,
+	uint8						WaveSize,
 	FVulkanSpirv&				Spirv,
 	const FString&				DebugName,
 	bool						bSourceContainsMetaDataOnly)
@@ -1311,6 +1312,9 @@ static void BuildShaderOutput(
 		NEWHeader.DebugName = ShaderInput.GenerateShaderName();
 	}
 
+	// Plug the passed in WaveSize
+	NEWHeader.WaveSize = WaveSize;
+
 	// Write out the header and shader source code.
 	FMemoryWriter Ar(ShaderOutput.ShaderCode.GetWriteAccess(), true);
 	Ar << NEWHeader;
@@ -1478,6 +1482,7 @@ static bool BuildShaderOutputFromSpirv(
 	FShaderCompilerOutput&					Output,
 	FVulkanBindingTable&					BindingTable,
 	const FString&							EntryPointName,
+	uint8									WaveSize,
 	bool									bStripReflect,
 	bool									bIsRayTracingShader,
 	bool									bDebugDump)
@@ -1888,6 +1893,7 @@ static bool BuildShaderOutputFromSpirv(
 		MetaData.Len(),
 		BindingTable,
 		ApproxInstructionCount,
+		WaveSize,
 		Spirv,
 		DebugName,
 		true // source contains meta data only
@@ -2140,6 +2146,96 @@ static void VulkanCreateDXCCompileBatchFiles(
 	}
 }
 
+// Quick and dirty way to get the location of the entrypoint in the source
+// NOTE: Preprocessed shaders have mcros resolves and comments removed, it makes this easier...
+static FString ParseEntrypointDecl(FStringView PreprocessedShader, FStringView Entrypoint)
+{
+	auto SkipWhitespace = [&](int32& Index)
+	{
+		while (FChar::IsWhitespace(PreprocessedShader[Index]))
+		{
+			++Index;
+		}
+	};
+
+	auto EraseDebugLines = [](FString& EntryPointDecl)
+	{
+		int32 HashIndex;
+		while (EntryPointDecl.FindChar(TEXT('#'), HashIndex))
+		{
+			while ((HashIndex < EntryPointDecl.Len()) && (!FChar::IsLinebreak(EntryPointDecl[HashIndex])))
+			{
+				EntryPointDecl[HashIndex] = TEXT(' ');
+				++HashIndex;
+			}
+		}
+	};
+
+	FString EntryPointDecl;
+
+	// Go through all the case sensitive matches in the source
+	int32 EntrypointIndex = PreprocessedShader.Find(Entrypoint);
+	check(EntrypointIndex != INDEX_NONE);
+	while (EntrypointIndex != INDEX_NONE)
+	{
+		// This should be the beginning of a new word
+		if ((EntrypointIndex == 0) || !FChar::IsWhitespace(PreprocessedShader[EntrypointIndex - 1]))
+		{
+			EntrypointIndex = PreprocessedShader.Find(Entrypoint, EntrypointIndex + 1);
+			continue;
+		}
+
+		// The next thing after the entrypoint should its parameters
+		// White space is allowed, so skip any that is found
+
+		int32 ParamsStart = EntrypointIndex + Entrypoint.Len();
+		SkipWhitespace(ParamsStart);
+		if (PreprocessedShader[ParamsStart] != TEXT('('))
+		{
+			EntrypointIndex = PreprocessedShader.Find(Entrypoint, ParamsStart);
+			continue;
+		}
+
+		int32 ParamsEnd = PreprocessedShader.Find(TEXT(")"), ParamsStart+1);
+		check(ParamsEnd != INDEX_NONE);
+		if (ParamsEnd == INDEX_NONE)
+		{
+			// Suspicious
+			EntrypointIndex = PreprocessedShader.Find(Entrypoint, ParamsStart);
+			continue;
+		}
+
+		// Make sure to grab everything up to the function content
+
+		int32 DeclEnd = ParamsEnd + 1;
+		while (PreprocessedShader[DeclEnd] != TEXT('{') && (PreprocessedShader[DeclEnd] != TEXT(';')))
+		{
+			++DeclEnd;
+		}
+		if (PreprocessedShader[DeclEnd] != TEXT('{'))
+		{
+			EntrypointIndex = PreprocessedShader.Find(Entrypoint, DeclEnd);
+			continue;
+		}
+
+		// Now back up to pick up the return value, the attributes and everything else that can come with it, like "[numthreads(1,1,1)]"
+
+		int32 DeclBegin = EntrypointIndex - 1;
+		while ( (DeclBegin > 0) && (PreprocessedShader[DeclBegin] != TEXT(';')) && (PreprocessedShader[DeclBegin] != TEXT('}')))
+		{
+			--DeclBegin;
+		}
+		++DeclBegin;
+
+		EntryPointDecl = FString(DeclEnd - DeclBegin, &PreprocessedShader[DeclBegin]);
+		EraseDebugLines(EntryPointDecl);
+		EntryPointDecl.TrimStartAndEndInline();
+		break;
+	}
+
+	return EntryPointDecl;
+}
+
 static bool CompileWithShaderConductor(
 	const FString&			PreprocessedShader,
 	const FString&			EntryPointName,
@@ -2195,6 +2291,45 @@ static bool CompileWithShaderConductor(
 			Options);
 	}
 
+	// Before the shader rewritter removes all traces of it, pull any WAVESIZE directives from the shader source
+	uint8 WaveSize = 0;
+	if (!bIsRayTracingShader)
+	{
+		const FString EntrypointDecl = ParseEntrypointDecl(PreprocessedShader, EntryPointName);
+
+		const FString WaveSizeMacro(TEXT("VULKAN_WAVESIZE("));
+		int32 WaveSizeIndex = EntrypointDecl.Find(*WaveSizeMacro, ESearchCase::CaseSensitive);
+		while (WaveSizeIndex != INDEX_NONE)
+		{
+			const int32 StartNumber = WaveSizeIndex + WaveSizeMacro.Len();
+			const int32 EndNumber = EntrypointDecl.Find(TEXT(")"), ESearchCase::CaseSensitive, ESearchDir::FromStart, StartNumber);
+			check(EndNumber != INDEX_NONE);
+
+			FString WaveSizeValue(EndNumber - StartNumber, &EntrypointDecl[StartNumber]);
+			WaveSizeValue.RemoveSpacesInline();
+			if (WaveSizeValue != TEXT("N"))  // skip the macro decl
+			{
+				float FloatResult = 0.0;
+				if (FMath::Eval(WaveSizeValue, FloatResult))
+				{
+					checkf((FloatResult >= 0.0f) && (FloatResult < (float)MAX_uint8), TEXT("Specified wave size is too large for 8bit uint!"));
+					WaveSize = static_cast<uint8>(FloatResult);
+
+				}
+				else
+				{
+					check(WaveSizeValue.IsNumeric());
+					const int32 ConvertedWaveSize = FCString::Atoi(*WaveSizeValue);
+					checkf((ConvertedWaveSize > 0) && (ConvertedWaveSize < MAX_uint8), TEXT("Specified wave size is too large for 8bit uint!"));
+					WaveSize = (uint8)ConvertedWaveSize;
+				}
+				break;
+			}
+
+			WaveSizeIndex = EntrypointDecl.Find(*WaveSizeMacro, ESearchCase::CaseSensitive, ESearchDir::FromStart, EndNumber);
+		}
+	}
+
 	if (bRewriteHlslSource)
 	{
 		// Rewrite HLSL source code to remove unused global resources and variables
@@ -2227,7 +2362,7 @@ static bool CompileWithShaderConductor(
 	Patch64bitSamplers(Spirv);
 
 	// Build shader output and binding table
-	Output.bSucceeded = BuildShaderOutputFromSpirv(CompilerContext, Spirv, Input, Output, BindingTable, EntryPointName, bStripReflect, bIsRayTracingShader, bDebugDump);
+	Output.bSucceeded = BuildShaderOutputFromSpirv(CompilerContext, Spirv, Input, Output, BindingTable, EntryPointName, WaveSize, bStripReflect, bIsRayTracingShader, bDebugDump);
 
 	if (Input.Environment.CompilerFlags.Contains(CFLAG_ExtraShaderData))
 	{
@@ -2322,6 +2457,7 @@ void DoCompileVulkanShader(const FShaderCompilerInput& Input, FShaderCompilerOut
 	if (TargetEnvironment >= CrossCompiler::FShaderConductorOptions::ETargetEnvironment::Vulkan_1_1)
 	{
 		AdditionalDefines.SetDefine(TEXT("PLATFORM_SUPPORTS_SM6_0_WAVE_OPERATIONS"), 1);
+		AdditionalDefines.SetDefine(TEXT("VULKAN_SUPPORTS_SUBGROUP_SIZE_CONTROL"), 1);
 	}
 	else
 	{
@@ -2406,7 +2542,6 @@ void DoCompileVulkanShader(const FShaderCompilerInput& Input, FShaderCompilerOut
 
 	UE::ShaderCompilerCommon::DumpDebugShaderData(Input, PreprocessedShaderSource, { CompilerInfo.CCFlags });
 
-	TArray<ANSICHAR> GeneratedGlslSource;
 	FVulkanBindingTable BindingTable(CompilerInfo.Frequency);
 	bool bSuccess = false;
 
diff --git a/Engine/Source/Runtime/VulkanRHI/Private/VulkanDevice.h b/Engine/Source/Runtime/VulkanRHI/Private/VulkanDevice.h
--- a/Engine/Source/Runtime/VulkanRHI/Private/VulkanDevice.h
+++ b/Engine/Source/Runtime/VulkanRHI/Private/VulkanDevice.h
@@ -81,6 +81,7 @@ struct FOptionalVulkanDeviceExtensions
 			uint64 HasEXTTextureCompressionASTCHDR : 1;
 			uint64 HasKHRMaintenance4 : 1;
 			uint64 HasKHRSynchronization2 : 1;
+			uint64 HasEXTSubgroupSizeControl : 1;
 		};
 		uint64 Packed;
 	};
diff --git a/Engine/Source/Runtime/VulkanRHI/Private/VulkanExtensions.cpp b/Engine/Source/Runtime/VulkanRHI/Private/VulkanExtensions.cpp
--- a/Engine/Source/Runtime/VulkanRHI/Private/VulkanExtensions.cpp
+++ b/Engine/Source/Runtime/VulkanRHI/Private/VulkanExtensions.cpp
@@ -978,6 +978,7 @@ public:
 	virtual void PostPhysicalDeviceFeatures(FOptionalVulkanDeviceExtensions& ExtensionFlags) override final
 	{
  		bSupportsSubgroupSizeControl = (SubgroupSizeControlFeatures.subgroupSizeControl == VK_TRUE);
+		ExtensionFlags.HasEXTSubgroupSizeControl = bSupportsSubgroupSizeControl ? 1 : 0;
 	}
 
 	virtual void PrePhysicalDeviceProperties(VkPhysicalDeviceProperties2KHR& PhysicalDeviceProperties2) override final
@@ -995,6 +996,8 @@ public:
 		{
 			GRHIMinimumWaveSize = SubgroupSizeControlProperties.minSubgroupSize;
 			GRHIMaximumWaveSize = SubgroupSizeControlProperties.maxSubgroupSize;
+
+			AddToPNext(DeviceCreateInfo, SubgroupSizeControlFeatures);
 		}
 	}
 
diff --git a/Engine/Source/Runtime/VulkanRHI/Private/VulkanPipeline.cpp b/Engine/Source/Runtime/VulkanRHI/Private/VulkanPipeline.cpp
--- a/Engine/Source/Runtime/VulkanRHI/Private/VulkanPipeline.cpp
+++ b/Engine/Source/Runtime/VulkanRHI/Private/VulkanPipeline.cpp
@@ -1261,6 +1261,7 @@ bool FVulkanPipelineStateCacheManager::CreateGfxPipelineFromEntry(FVulkanRHIGrap
 	PipelineInfo.pStages = ShaderStages;
 	// main_00000000_00000000
 	ANSICHAR EntryPoints[ShaderStage::NumStages][24];
+	VkPipelineShaderStageRequiredSubgroupSizeCreateInfo RequiredSubgroupSizeCreateInfo[ShaderStage::NumStages];
 	for (int32 ShaderStage = 0; ShaderStage < ShaderStage::NumStages; ++ShaderStage)
 	{
 		if (!ShaderModules[ShaderStage].IsValid())
@@ -1275,6 +1276,16 @@ bool FVulkanPipelineStateCacheManager::CreateGfxPipelineFromEntry(FVulkanRHIGrap
 		ShaderStages[PipelineInfo.stageCount].module = ShaderModules[CurrStage]->GetVkShaderModule();
 		Shaders[ShaderStage]->GetEntryPoint(EntryPoints[PipelineInfo.stageCount], 24);
 		ShaderStages[PipelineInfo.stageCount].pName = EntryPoints[PipelineInfo.stageCount];
+
+		const FVulkanShaderHeader& ShaderHeader = Shaders[CurrStage]->GetCodeHeader();
+		if ((ShaderHeader.WaveSize > 0) && Device->GetOptionalExtensions().HasEXTSubgroupSizeControl)
+		{
+			check((ShaderHeader.WaveSize >= GRHIMinimumWaveSize) && (ShaderHeader.WaveSize <= GRHIMaximumWaveSize));
+			ZeroVulkanStruct(RequiredSubgroupSizeCreateInfo[PipelineInfo.stageCount], VK_STRUCTURE_TYPE_PIPELINE_SHADER_STAGE_REQUIRED_SUBGROUP_SIZE_CREATE_INFO);
+			RequiredSubgroupSizeCreateInfo[PipelineInfo.stageCount].requiredSubgroupSize = ShaderHeader.WaveSize;
+			ShaderStages[PipelineInfo.stageCount].pNext = &RequiredSubgroupSizeCreateInfo[PipelineInfo.stageCount];
+		}
+
 		PipelineInfo.stageCount++;
 	}
 
@@ -2138,6 +2149,16 @@ FVulkanComputePipeline* FVulkanPipelineStateCacheManager::CreateComputePipelineF
 	Shader->GetEntryPoint(EntryPoint, 24);
 	PipelineInfo.stage.pName = EntryPoint;
 	PipelineInfo.layout = ComputeLayout->GetPipelineLayout();
+
+	VkPipelineShaderStageRequiredSubgroupSizeCreateInfo RequiredSubgroupSizeCreateInfo;
+	if ((CSHeader.WaveSize > 0) && Device->GetOptionalExtensions().HasEXTSubgroupSizeControl)
+	{
+		check((CSHeader.WaveSize >= GRHIMinimumWaveSize) && (CSHeader.WaveSize <= GRHIMaximumWaveSize));
+		ZeroVulkanStruct(RequiredSubgroupSizeCreateInfo, VK_STRUCTURE_TYPE_PIPELINE_SHADER_STAGE_REQUIRED_SUBGROUP_SIZE_CREATE_INFO);
+		RequiredSubgroupSizeCreateInfo.requiredSubgroupSize = CSHeader.WaveSize;
+		PipelineInfo.stage.pNext = &RequiredSubgroupSizeCreateInfo;
+	}
+
 	VkResult Result;
 	{
 		QUICK_SCOPE_CYCLE_COUNTER(STAT_VulkanComputePSOCreate);
diff --git a/Engine/Source/Runtime/VulkanRHI/Public/VulkanShaderResources.h b/Engine/Source/Runtime/VulkanRHI/Public/VulkanShaderResources.h
--- a/Engine/Source/Runtime/VulkanRHI/Public/VulkanShaderResources.h
+++ b/Engine/Source/Runtime/VulkanRHI/Public/VulkanShaderResources.h
@@ -188,6 +188,7 @@ struct FVulkanShaderHeader
 
 	FSHAHash								SourceHash;
 	uint32									SpirvCRC = 0;
+	uint8									WaveSize = 0;
 
 	TArray<FSpirvInfo>						UniformBufferSpirvInfos;
 	TArray<FSpirvInfo>						GlobalSpirvInfos;
@@ -291,6 +292,7 @@ inline FArchive& operator<<(FArchive& Ar, FVulkanShaderHeader& Header)
 	Ar << Header.RayTracingPayloadType;
 	Ar << Header.SourceHash;
 	Ar << Header.SpirvCRC;
+	Ar << Header.WaveSize;
 	Ar << Header.UniformBufferSpirvInfos;
 	Ar << Header.GlobalSpirvInfos;
 	Ar << Header.DebugName;
