VulkanRHI:  Override ShaderParameterParser for Vulkan bindless resource declarations.  Use a typedef for each resource's type and declare its variable array with that type, at the end (in SPIRV) collapse all the arrays of a same descriptor type on the same binding.  Move 'per-stage' uniform buffer bindings into descriptor sets AFTER the descrciptor sets for the resource arrays (since the resource array bindings will never change, but the ubo will change at a very high frequency).  Also add the possibility for platform common headers to overload UB_RESOURCE_MEMBER_RESOURCE to match the altered declaration in ShaderParameterParser (temporary solution).  The result gets around the lack of ResourceDescriptorHeap in Vulkan.    (Vulkan Bindless WIP)

#rnx

diff --git a/Engine/Shaders/Public/BindlessResources.ush b/Engine/Shaders/Public/BindlessResources.ush
--- a/Engine/Shaders/Public/BindlessResources.ush
+++ b/Engine/Shaders/Public/BindlessResources.ush
@@ -37,7 +37,9 @@
 	#define AUTO_BINDLESS_RESOURCE_INDEX(Name)          uint BindlessResource_##Name
 	#define AUTO_BINDLESS_RESOURCE_VARIABLE(Type, Name) Type GetBindlessResource##Name() { return GetResourceFromHeap(Type, BindlessResource_##Name); }; static const Type Name = GetBindlessResource##Name()
 
-	#define UB_RESOURCE_MEMBER_RESOURCE(Type, Name) AUTO_BINDLESS_RESOURCE_VARIABLE(Type, Name)
+	#ifndef UB_RESOURCE_MEMBER_RESOURCE  // allow platforms to override ub member resource decl
+		#define UB_RESOURCE_MEMBER_RESOURCE(Type, Name) AUTO_BINDLESS_RESOURCE_VARIABLE(Type, Name)
+	#endif
 #else
 	#define DEFINE_RESOURCE_VARIABLE(Type, Name)  Type Name
 	#define DEFINE_RESOURCE(Type, Name)           Type Name
diff --git a/Engine/Shaders/Public/Platform/Vulkan/VulkanCommon.ush b/Engine/Shaders/Public/Platform/Vulkan/VulkanCommon.ush
--- a/Engine/Shaders/Public/Platform/Vulkan/VulkanCommon.ush
+++ b/Engine/Shaders/Public/Platform/Vulkan/VulkanCommon.ush
@@ -5,7 +5,7 @@
 =============================================================================*/
 
 // Update this GUID to improve shader recompilation for Vulkan only shaders
-// GUID = 1084d248-f5cb-41ba-8ff7-3ee0752a4295
+// GUID = d8d5a9ec-19c5-417c-b1ca-876d31a6809a
 
 #pragma once
 
@@ -42,8 +42,11 @@ void ImageInterlockedMaxUInt64(RWTexture2D<UlongType> Dest, uint2 Coord, UlongTy
 	COMPILER_DEFINE VULKAN_SAMPLER_HEAP(Type) VULKAN_HEAP_NAME(VULKAN_BINDLESS_SAMPLER_ARRAY_PREFIX,Type)
 	COMPILER_DEFINE VULKAN_RESOURCE_HEAP(Type) VULKAN_HEAP_NAME(VULKAN_BINDLESS_RESOURCE_ARRAY_PREFIX,Type)
 
-	// todo-jn: bindless resource heaps
-	COMPILER_DEFINE GetResourceFromHeap(Type, Index) VULKAN_RESOURCE_HEAP(Type)[Index]
+	// Bindless resource heaps are filled in by the shader compiler
+	COMPILER_DEFINE GetResourceFromHeap(Type, Index)  /* filled by compiler */
+
+	// Override the UB_RESOURCE_MEMBER_RESOURCE to use the Vulkan resource heaps
+	#define UB_RESOURCE_MEMBER_RESOURCE(Type, Name) typedef Type SafeType##Name; SafeType##Name VULKAN_RESOURCE_HEAP(SafeType##Name)[]; static const Type Name = VULKAN_RESOURCE_HEAP(SafeType##Name)[BindlessResource_##Name]
 
 	// bindless sampler heaps (specify a register so AUTO_BINDLESS macros aren't applied on them)
 	SamplerState VULKAN_SAMPLER_HEAP(SamplerState)[] : register(s0, space1000);
diff --git a/Engine/Source/Developer/ShaderCompilerCommon/Public/ShaderParameterParser.h b/Engine/Source/Developer/ShaderCompilerCommon/Public/ShaderParameterParser.h
--- a/Engine/Source/Developer/ShaderCompilerCommon/Public/ShaderParameterParser.h
+++ b/Engine/Source/Developer/ShaderCompilerCommon/Public/ShaderParameterParser.h
@@ -47,7 +47,6 @@ public:
 			return Member != nullptr;
 		}
 
-	private:
 		int32 ParsedPragmaLineoffset = 0;
 		int32 ParsedLineOffset = 0;
 
@@ -141,7 +140,7 @@ public:
 		return ExtractFileAndLine(ParsedParameter.ParsedPragmaLineoffset, ParsedParameter.ParsedLineOffset, OutFile, OutLine);
 	}
 
-private:
+protected:
 	/** Parses the preprocessed shader code */
 	bool ParseParameters(
 		const FShaderCompilerInput& CompilerInput,
diff --git a/Engine/Source/Developer/VulkanShaderFormat/Private/VulkanShaderCompiler.cpp b/Engine/Source/Developer/VulkanShaderFormat/Private/VulkanShaderCompiler.cpp
--- a/Engine/Source/Developer/VulkanShaderFormat/Private/VulkanShaderCompiler.cpp
+++ b/Engine/Source/Developer/VulkanShaderFormat/Private/VulkanShaderCompiler.cpp
@@ -1428,7 +1428,21 @@ static void GatherSpirvReflectionBindings(
 		// Remove sampler heaps from binding arrays
 		MoveBindlessHeaps(OutBindings.Samplers, VulkanBindless::kBindlessSamplerArrayPrefix, VulkanBindless::BindlessSamplerSet);
 
-		// todo-jn: Remove resource heaps from binding arrays
+		// Remove resource heaps from binding arrays
+		MoveBindlessHeaps(OutBindings.SBufferUAVs, VulkanBindless::kBindlessResourceArrayPrefix, VulkanBindless::BindlessStorageBufferSet);
+		MoveBindlessHeaps(OutBindings.TextureSRVs, VulkanBindless::kBindlessResourceArrayPrefix, VulkanBindless::BindlessSampledImageSet);
+		MoveBindlessHeaps(OutBindings.TextureUAVs, VulkanBindless::kBindlessResourceArrayPrefix, VulkanBindless::BindlessStorageImageSet);
+		MoveBindlessHeaps(OutBindings.TBufferSRVs, VulkanBindless::kBindlessResourceArrayPrefix, VulkanBindless::BindlessUniformTexelBufferSet);
+		MoveBindlessHeaps(OutBindings.TBufferUAVs, VulkanBindless::kBindlessResourceArrayPrefix, VulkanBindless::BindlessStorageTexelBufferSet);
+		MoveBindlessHeaps(OutBindings.AccelerationStructures, VulkanBindless::kBindlessResourceArrayPrefix, VulkanBindless::BindlessSamplerSet);
+
+		// Move uniform buffers to the correct set
+		const uint32 BinldessDescSetNo = VulkanBindless::NumBindlessSets + ShaderStage::GetStageForFrequency(ShaderFrequency);
+		for (int32 Index = OutBindings.UniformBuffers.Num() - 1; Index >= 0; --Index)
+		{
+			const SpvReflectDescriptorBinding* pBinding = OutBindings.UniformBuffers[Index];
+			Reflection.ChangeDescriptorBindingNumbers(pBinding, pBinding->binding, BinldessDescSetNo);
+		}
 	}
 }
 
@@ -1671,9 +1685,11 @@ static bool BuildShaderOutputFromSpirv(
 		if (ResourceName == UBOGlobalsNameSpv)
 		{
 			// Register binding for uniform buffer
-			int32 BindingIndex = MapDescriptorBindingToIndex(Binding);
+			const int32 BindingIndex = MapDescriptorBindingToIndex(Binding);
+			const uint32 BindlessDescSetNumber = VulkanBindless::NumBindlessSets + ShaderStage::GetStageForFrequency(static_cast<EShaderFrequency>(Input.Target.Frequency));
+			const uint32_t DescSetNumber = bSupportsBindless ? BindlessDescSetNumber : SPV_REFLECT_SET_NUMBER_DONT_CHANGE;
 
-			SpvResult = Reflection.ChangeDescriptorBindingNumbers(Binding, BindingIndex);// , GlobalSetId);
+			SpvResult = Reflection.ChangeDescriptorBindingNumbers(Binding, BindingIndex, DescSetNumber);
 			check(SpvResult == SPV_REFLECT_RESULT_SUCCESS);
 
 			Spirv.ReflectionInfo.Add(FVulkanSpirv::FEntry(UBOGlobalsNameSpv, BindingIndex));
@@ -2254,7 +2270,6 @@ static bool CompileWithShaderConductor(
 	const FShaderCompilerInput& Input = CompilerInfo.Input;
 
 	const bool bIsRayTracingShader = Input.IsRayTracingShader();
-	const bool bHasBindless = Input.Environment.CompilerFlags.Contains(CFLAG_BindlessResources) || Input.Environment.CompilerFlags.Contains(CFLAG_BindlessSamplers);
 	const bool bRewriteHlslSource = !bIsRayTracingShader;
 	const bool bDebugDump = Input.DumpDebugInfoEnabled();
 
@@ -2388,6 +2403,56 @@ static bool CompileWithShaderConductor(
 #endif // PLATFORM_MAC || PLATFORM_WINDOWS || PLATFORM_LINUX
 
 
+// Overload parameter declaration from the default FShaderParameterParser in order to take into account Vulkan particularities for resources
+class FVulkanShaderParameterParser : public FShaderParameterParser
+{
+public:
+	FVulkanShaderParameterParser(bool bHasBindlessResources)
+		: bEnabled(bHasBindlessResources)
+	{}
+
+protected:
+	FString GenerateBindlessParameterDeclaration(const FParsedShaderParameter& ParsedParameter) const override
+	{
+		if (bEnabled && (ParsedParameter.BindlessConversionType == EBindlessConversionType::Resource))
+		{
+			const TCHAR* StorageClass = ParsedParameter.bGloballyCoherent ? TEXT("globallycoherent ") : TEXT("");
+
+			const FStringView Name = ParsedParameter.ParsedName;
+			const FStringView Type = ParsedParameter.ParsedType;
+
+			const FString RewriteType = FString::Printf(TEXT("SafeType%.*s"), Name.Len(), Name.GetData());
+
+			TStringBuilder<512> Result;
+
+			// If we weren't going to be added to a root constant buffer, that means we need to declare our index before we declare our getter.
+			if (ParsedParameter.ConstantBufferParameterType == EShaderParameterType::Num)
+			{
+				// e.g. `uint BindlessResource_##Name;`
+				Result << TEXT("uint BindlessResource_") << Name << TEXT("; ");
+			}
+
+			// Add the typedef
+			Result << TEXT("typedef ") << Type << TEXT(" ") << RewriteType << TEXT("; ");
+
+			// Declare a heap for the RewriteType
+			// e.g. `SafeType##Name ResourceDescriptorHeap_SafeType##Name[];`
+			Result << RewriteType << TEXT(" ") << VulkanBindless::kBindlessResourceArrayPrefix << RewriteType << TEXT("[]; ");
+			// :todo-jn: specify the descripor set and binding directly in source instead of patching SPIRV
+
+			// e.g. `static const Type Name = GetBindlessResource##Name()`
+			Result << TEXT("static const ") << StorageClass << RewriteType << TEXT(" ") << Name << TEXT(" = ") << VulkanBindless::kBindlessResourceArrayPrefix << RewriteType << TEXT("[BindlessResource_") << Name << TEXT("];");
+
+			return Result.ToString();
+		}
+		return FShaderParameterParser::GenerateBindlessParameterDeclaration(ParsedParameter);
+	}
+
+private:
+	const bool bEnabled;
+};
+
+
 void DoCompileVulkanShader(const FShaderCompilerInput& Input, FShaderCompilerOutput& Output, const class FString& WorkingDirectory, EVulkanShaderVersion Version)
 {
 	check(IsVulkanShaderFormat(Input.ShaderFormat));
@@ -2497,7 +2562,8 @@ void DoCompileVulkanShader(const FShaderCompilerInput& Input, FShaderCompilerOut
 		}
 	}
 
-	FShaderParameterParser ShaderParameterParser;
+	const bool bHasBindlessResources = Input.Environment.CompilerFlags.Contains(CFLAG_BindlessResources);
+	FVulkanShaderParameterParser ShaderParameterParser(bHasBindlessResources);
 	if (!ShaderParameterParser.ParseAndModify(Input, Output, PreprocessedShaderSource))
 	{
 		// The FShaderParameterParser will add any relevant errors.
diff --git a/Engine/Source/Runtime/VulkanRHI/Public/VulkanCommon.h b/Engine/Source/Runtime/VulkanRHI/Public/VulkanCommon.h
--- a/Engine/Source/Runtime/VulkanRHI/Public/VulkanCommon.h
+++ b/Engine/Source/Runtime/VulkanRHI/Public/VulkanCommon.h
@@ -116,10 +116,19 @@ namespace VulkanBindless
 	enum EDescriptorSets
 	{
 		BindlessSamplerSet = 0,
-		BindlessResourceSet = 1,
 
-		NumBindlessSets = 2,
-		MaxNumSets = ShaderStage::EStage::MaxNumSets + NumBindlessSets
+		BindlessStorageBufferSet,
+
+		BindlessStorageImageSet,
+		BindlessSampledImageSet,
+
+		BindlessStorageTexelBufferSet,
+		BindlessUniformTexelBufferSet,
+
+		BindlessAccelerationStructureSet,
+
+		NumBindlessSets,
+		MaxNumSets = ShaderStage::EStage::NumStages + NumBindlessSets
 	};
 
 	// Prefix used to declare arrays of samplers/resources for bindless
