VulkanRHI: Use a single descriptor set for all uniform buffers in Vulkan bindless, separate the stages with offsets inside that single decriptor set.  Gets around previous issue with different descriptor sets were using the same descriptor buffer with different offsets.  (Vulkan Bindless WIP)

#rnx

diff --git a/Engine/Shaders/Public/Platform/Vulkan/VulkanCommon.ush b/Engine/Shaders/Public/Platform/Vulkan/VulkanCommon.ush
--- a/Engine/Shaders/Public/Platform/Vulkan/VulkanCommon.ush
+++ b/Engine/Shaders/Public/Platform/Vulkan/VulkanCommon.ush
@@ -5,7 +5,7 @@
 =============================================================================*/
 
 // Update this GUID to improve shader recompilation for Vulkan only shaders
-// GUID = 0207DEF7-0E8C-4D8C-8170-52C7B9864E3F
+// GUID = 83946375-67fe-4f8f-be31-f9c03b9426ad
 
 #pragma once
 
diff --git a/Engine/Source/Developer/VulkanShaderFormat/Private/VulkanShaderCompiler.cpp b/Engine/Source/Developer/VulkanShaderFormat/Private/VulkanShaderCompiler.cpp
--- a/Engine/Source/Developer/VulkanShaderFormat/Private/VulkanShaderCompiler.cpp
+++ b/Engine/Source/Developer/VulkanShaderFormat/Private/VulkanShaderCompiler.cpp
@@ -918,7 +918,7 @@ static void BuildShaderOutput(
 
 	FOLDVulkanCodeHeader OLDHeader;
 
-	EShaderFrequency Frequency = (EShaderFrequency)ShaderOutput.Target.Frequency;
+	const EShaderFrequency Frequency = (EShaderFrequency)ShaderOutput.Target.Frequency;
 
 	TBitArray<> UsedUniformBufferSlots;
 	UsedUniformBufferSlots.Init(false, 32);
@@ -984,6 +984,10 @@ static void BuildShaderOutput(
 		}
 	}
 
+
+	const bool bSupportsBindless = ShaderInput.Environment.CompilerFlags.Contains(CFLAG_BindlessResources) || ShaderInput.Environment.CompilerFlags.Contains(CFLAG_BindlessSamplers);
+	const int32 StageOffset = bSupportsBindless ? (ShaderStage::GetStageForFrequency(static_cast<EShaderFrequency>(ShaderInput.Target.Frequency)) * VulkanBindless::MaxUniformBuffersPerStage) : 0;
+
 	const TArray<FVulkanBindingTable::FBinding>& HlslccBindings = BindingTable.GetBindings();
 	OLDHeader.NEWDescriptorInfo.NumBufferInfos = 0;
 	OLDHeader.NEWDescriptorInfo.NumImageInfos = 0;
@@ -1010,7 +1014,7 @@ static void BuildShaderOutput(
 			{
 				FOLDVulkanCodeHeader::FPackedUBToVulkanBindingIndex* New = new(OLDHeader.NEWPackedUBToVulkanBindingIndices) FOLDVulkanCodeHeader::FPackedUBToVulkanBindingIndex;
 				New->TypeName = (CrossCompiler::EPackedTypeName)Binding.SubType;
-				New->VulkanBindingIndex = Index;
+				New->VulkanBindingIndex = StageOffset + Index;
 				++OLDHeader.NEWDescriptorInfo.NumBufferInfos;
 			}
 			break;
@@ -1381,7 +1385,8 @@ static void GatherSpirvReflectionBindings(
 	uint32 NumDescriptorSets = 0;
 
 	// If bindless is supported, then offset the descriptor set to fit the bindless heaps at the beginning
-	const uint32 DescSetNo = bSupportsBindless ? VulkanBindless::NumBindlessSets + ShaderStage::GetStageForFrequency(ShaderFrequency) : ShaderStage::GetStageForFrequency(ShaderFrequency);
+	const uint32 StageIndex = (uint32)ShaderStage::GetStageForFrequency(ShaderFrequency);
+	const uint32 DescSetNo = bSupportsBindless ? VulkanBindless::NumBindlessSets + StageIndex : StageIndex;
 
 	SpvReflectResult SpvResult = Reflection.EnumerateDescriptorSets(&NumDescriptorSets, nullptr);
 	check(SpvResult == SPV_REFLECT_RESULT_SUCCESS);
@@ -1442,11 +1447,14 @@ static void GatherSpirvReflectionBindings(
 		MoveBindlessHeaps(OutBindings.AccelerationStructures, VulkanBindless::kBindlessResourceArrayPrefix, VulkanBindless::BindlessSamplerSet);
 
 		// Move uniform buffers to the correct set
-		const uint32 BinldessDescSetNo = VulkanBindless::NumBindlessSets + ShaderStage::GetStageForFrequency(ShaderFrequency);
-		for (int32 Index = OutBindings.UniformBuffers.Num() - 1; Index >= 0; --Index)
 		{
-			const SpvReflectDescriptorBinding* pBinding = OutBindings.UniformBuffers[Index];
-			Reflection.ChangeDescriptorBindingNumbers(pBinding, pBinding->binding, BinldessDescSetNo);
+			const uint32 BindingOffset = (StageIndex * VulkanBindless::MaxUniformBuffersPerStage);
+			const uint32 BinldessDescSetNo = VulkanBindless::BindlessUniformBufferSet;
+			for (int32 Index = OutBindings.UniformBuffers.Num() - 1; Index >= 0; --Index)
+			{
+				const SpvReflectDescriptorBinding* pBinding = OutBindings.UniformBuffers[Index];
+				Reflection.ChangeDescriptorBindingNumbers(pBinding, BindingOffset + pBinding->binding, BinldessDescSetNo);
+			}
 		}
 	}
 }
@@ -1690,9 +1698,9 @@ static bool BuildShaderOutputFromSpirv(
 		if (ResourceName == UBOGlobalsNameSpv)
 		{
 			// Register binding for uniform buffer
-			const int32 BindingIndex = MapDescriptorBindingToIndex(Binding);
-			const uint32 BindlessDescSetNumber = VulkanBindless::NumBindlessSets + ShaderStage::GetStageForFrequency(static_cast<EShaderFrequency>(Input.Target.Frequency));
-			const uint32_t DescSetNumber = bSupportsBindless ? BindlessDescSetNumber : SPV_REFLECT_SET_NUMBER_DONT_CHANGE;
+			const int32 StageOffset = bSupportsBindless ? (ShaderStage::GetStageForFrequency(static_cast<EShaderFrequency>(Input.Target.Frequency)) * VulkanBindless::MaxUniformBuffersPerStage) : 0;
+			const int32 BindingIndex = MapDescriptorBindingToIndex(Binding) + StageOffset;
+			const uint32_t DescSetNumber = bSupportsBindless ? VulkanBindless::BindlessUniformBufferSet : SPV_REFLECT_SET_NUMBER_DONT_CHANGE;
 
 			SpvResult = Reflection.ChangeDescriptorBindingNumbers(Binding, BindingIndex, DescSetNumber);
 			check(SpvResult == SPV_REFLECT_RESULT_SUCCESS);
@@ -1722,9 +1730,10 @@ static bool BuildShaderOutputFromSpirv(
 		if (ResourceName != UBOGlobalsNameSpv)
 		{
 			// Register uniform buffer
-
-			int32 BindingIndex = MapDescriptorBindingToIndex(Binding);
-			SpvResult = Reflection.ChangeDescriptorBindingNumbers(Binding, BindingIndex);//, GlobalSetId);
+			const int32 StageOffset = bSupportsBindless ? (ShaderStage::GetStageForFrequency(static_cast<EShaderFrequency>(Input.Target.Frequency)) * VulkanBindless::MaxUniformBuffersPerStage) : 0;
+			const int32 BindingIndex = MapDescriptorBindingToIndex(Binding) + StageOffset;
+			const uint32_t DescSetNumber = bSupportsBindless ? VulkanBindless::BindlessUniformBufferSet : SPV_REFLECT_SET_NUMBER_DONT_CHANGE;
+			SpvResult = Reflection.ChangeDescriptorBindingNumbers(Binding, BindingIndex, DescSetNumber);
 			check(SpvResult == SPV_REFLECT_RESULT_SUCCESS);
 
 			Spirv.ReflectionInfo.Add(FVulkanSpirv::FEntry(ResourceName, BindingIndex));
diff --git a/Engine/Source/Runtime/VulkanRHI/Private/VulkanDescriptorSets.h b/Engine/Source/Runtime/VulkanRHI/Private/VulkanDescriptorSets.h
--- a/Engine/Source/Runtime/VulkanRHI/Private/VulkanDescriptorSets.h
+++ b/Engine/Source/Runtime/VulkanRHI/Private/VulkanDescriptorSets.h
@@ -1262,11 +1262,10 @@ private:
 	};
 	BindlessSetState BindlessSetStates[VulkanBindless::NumBindlessSets];
 
-	BindlessSetState BindlessUniformBufferSetState;
-	std::atomic<uint32> UniformBlockIndex = 0;
+	std::atomic<uint32> CurrentUniformBufferDescriptorIndex = 0;
 
-	VkDescriptorBufferBindingInfoEXT BufferBindingInfo[VulkanBindless::NumBindlessSets + 1];
-	uint32_t BufferIndices[VulkanBindless::MaxNumSets];
+	VkDescriptorBufferBindingInfoEXT BufferBindingInfo[VulkanBindless::NumBindlessSets];
+	uint32_t BufferIndices[VulkanBindless::NumBindlessSets];
 
 	VkPipelineLayout BindlessPipelineLayout = VK_NULL_HANDLE;
 
diff --git a/Engine/Source/Runtime/VulkanRHI/Private/VulkanShaders.cpp b/Engine/Source/Runtime/VulkanRHI/Private/VulkanShaders.cpp
--- a/Engine/Source/Runtime/VulkanRHI/Private/VulkanShaders.cpp
+++ b/Engine/Source/Runtime/VulkanRHI/Private/VulkanShaders.cpp
@@ -391,10 +391,15 @@ void FVulkanShader::PurgeShaderModules()
 
 void FVulkanLayout::PatchSpirvBindings(FVulkanShader::FSpirvCode& SprivCode, EShaderFrequency Frequency, const FVulkanShaderHeader& CodeHeader, VkShaderStageFlagBits InStageFlag) const
 {
+	// Bindless shader compilation already places descriptors and bindings in their fixed values based on stage and descriptor type
+	if (Device->SupportsBindless())
+	{
+		return;
+	}
+
 	TArrayView<uint32> Spirv = SprivCode.GetCodeView();	//#todo-rco: Do we need an actual copy of the SPIR-V?
 	ShaderStage::EStage Stage = ShaderStage::GetStageForFrequency(Frequency);
 	const FDescriptorSetRemappingInfo::FStageInfo& StageInfo = DescriptorSetLayout.RemappingInfo.StageInfos[Stage];
-	const uint32 FirstSet = Device->SupportsBindless() ? VulkanBindless::NumBindlessSets : 0;
 
 	checkSlow(StageInfo.UniformBuffers.Num() == CodeHeader.UniformBufferSpirvInfos.Num());
 	for (int32 Index = 0; Index < CodeHeader.UniformBufferSpirvInfos.Num(); ++Index)
@@ -405,7 +410,7 @@ void FVulkanLayout::PatchSpirvBindings(FVulkanShader::FSpirvCode& SprivCode, ESh
 			const uint32 OffsetBindingIndex = CodeHeader.UniformBufferSpirvInfos[Index].BindingIndexOffset;
 			check(OffsetDescriptorSet != UINT32_MAX && OffsetBindingIndex != UINT32_MAX);
 			uint16 NewDescriptorSet = StageInfo.UniformBuffers[Index].Remapping.NewDescriptorSet;
-			Spirv[OffsetDescriptorSet] = FirstSet + NewDescriptorSet;
+			Spirv[OffsetDescriptorSet] = NewDescriptorSet;
 			uint16 NewBindingIndex = StageInfo.UniformBuffers[Index].Remapping.NewBindingIndex;
 			Spirv[OffsetBindingIndex] = NewBindingIndex;
 		}
@@ -418,7 +423,7 @@ void FVulkanLayout::PatchSpirvBindings(FVulkanShader::FSpirvCode& SprivCode, ESh
 		const uint32 OffsetBindingIndex = CodeHeader.GlobalSpirvInfos[Index].BindingIndexOffset;
 		check(OffsetDescriptorSet != UINT32_MAX && OffsetBindingIndex != UINT32_MAX);
 		uint16 NewDescriptorSet = StageInfo.Globals[Index].NewDescriptorSet;
-		Spirv[OffsetDescriptorSet] = FirstSet + NewDescriptorSet;
+		Spirv[OffsetDescriptorSet] = NewDescriptorSet;
 		uint16 NewBindingIndex = StageInfo.Globals[Index].NewBindingIndex;
 		Spirv[OffsetBindingIndex] = NewBindingIndex;
 	}
@@ -429,7 +434,7 @@ void FVulkanLayout::PatchSpirvBindings(FVulkanShader::FSpirvCode& SprivCode, ESh
 		const uint32 OffsetDescriptorSet = CodeHeader.PackedUBs[Index].SPIRVDescriptorSetOffset;
 		const uint32 OffsetBindingIndex = CodeHeader.PackedUBs[Index].SPIRVBindingIndexOffset;
 		check(OffsetDescriptorSet != UINT32_MAX && OffsetBindingIndex != UINT32_MAX);
-		Spirv[OffsetDescriptorSet] = FirstSet + StageInfo.PackedUBDescriptorSet;
+		Spirv[OffsetDescriptorSet] = StageInfo.PackedUBDescriptorSet;
 		Spirv[OffsetBindingIndex] = StageInfo.PackedUBBindingIndices[Index];
 	}
 }
diff --git a/Engine/Source/Runtime/VulkanRHI/Public/VulkanCommon.h b/Engine/Source/Runtime/VulkanRHI/Public/VulkanCommon.h
--- a/Engine/Source/Runtime/VulkanRHI/Public/VulkanCommon.h
+++ b/Engine/Source/Runtime/VulkanRHI/Public/VulkanCommon.h
@@ -113,6 +113,8 @@ namespace ShaderStage
 
 namespace VulkanBindless
 {
+	static constexpr uint32 MaxUniformBuffersPerStage = 8;
+
 	enum EDescriptorSets
 	{
 		BindlessSamplerSet = 0,
@@ -127,8 +129,9 @@ namespace VulkanBindless
 
 		BindlessAccelerationStructureSet,
 
+		BindlessUniformBufferSet,  // Keep last
 		NumBindlessSets,
-		MaxNumSets = ShaderStage::EStage::NumStages + NumBindlessSets
+		MaxNumSets = NumBindlessSets
 	};
 
 	// Prefix used to declare arrays of samplers/resources for bindless
