VulkanRHI: Fix static analysis warning.

#rnx

diff --git a/Engine/Source/Developer/VulkanShaderFormat/Private/VulkanShaderCompiler.cpp b/Engine/Source/Developer/VulkanShaderFormat/Private/VulkanShaderCompiler.cpp
--- a/Engine/Source/Developer/VulkanShaderFormat/Private/VulkanShaderCompiler.cpp
+++ b/Engine/Source/Developer/VulkanShaderFormat/Private/VulkanShaderCompiler.cpp
@@ -1700,7 +1700,7 @@ static bool BuildShaderOutputFromSpirv(
 			// Register binding for uniform buffer
 			const int32 StageOffset = bSupportsBindless ? (ShaderStage::GetStageForFrequency(static_cast<EShaderFrequency>(Input.Target.Frequency)) * VulkanBindless::MaxUniformBuffersPerStage) : 0;
 			const int32 BindingIndex = MapDescriptorBindingToIndex(Binding) + StageOffset;
-			const uint32_t DescSetNumber = bSupportsBindless ? VulkanBindless::BindlessUniformBufferSet : SPV_REFLECT_SET_NUMBER_DONT_CHANGE;
+			const uint32_t DescSetNumber = bSupportsBindless ? (uint32_t)VulkanBindless::BindlessUniformBufferSet : (uint32_t)SPV_REFLECT_SET_NUMBER_DONT_CHANGE;
 
 			SpvResult = Reflection.ChangeDescriptorBindingNumbers(Binding, BindingIndex, DescSetNumber);
 			check(SpvResult == SPV_REFLECT_RESULT_SUCCESS);
@@ -1732,7 +1732,7 @@ static bool BuildShaderOutputFromSpirv(
 			// Register uniform buffer
 			const int32 StageOffset = bSupportsBindless ? (ShaderStage::GetStageForFrequency(static_cast<EShaderFrequency>(Input.Target.Frequency)) * VulkanBindless::MaxUniformBuffersPerStage) : 0;
 			const int32 BindingIndex = MapDescriptorBindingToIndex(Binding) + StageOffset;
-			const uint32_t DescSetNumber = bSupportsBindless ? VulkanBindless::BindlessUniformBufferSet : SPV_REFLECT_SET_NUMBER_DONT_CHANGE;
+			const uint32_t DescSetNumber = bSupportsBindless ? (uint32_t)VulkanBindless::BindlessUniformBufferSet : (uint32_t)SPV_REFLECT_SET_NUMBER_DONT_CHANGE;
 			SpvResult = Reflection.ChangeDescriptorBindingNumbers(Binding, BindingIndex, DescSetNumber);
 			check(SpvResult == SPV_REFLECT_RESULT_SUCCESS);
 
