VulkanRHI: Fix issue where incorrect TargetEnvironment was selected when previewing other SM5 modes.

diff --git a/Engine/Source/Developer/VulkanShaderFormat/Private/VulkanShaderCompiler.cpp b/Engine/Source/Developer/VulkanShaderFormat/Private/VulkanShaderCompiler.cpp
--- a/Engine/Source/Developer/VulkanShaderFormat/Private/VulkanShaderCompiler.cpp
+++ b/Engine/Source/Developer/VulkanShaderFormat/Private/VulkanShaderCompiler.cpp
@@ -65,10 +65,10 @@ inline bool SupportsOfflineCompiler(FName ShaderFormat)
 		|| ShaderFormat == NAME_VULKAN_SM5_ANDROID;
 }
 
-inline CrossCompiler::FShaderConductorOptions::ETargetEnvironment GetMinimumTargetEnvironment(EShaderPlatform ShaderPlatform)
+inline CrossCompiler::FShaderConductorOptions::ETargetEnvironment GetMinimumTargetEnvironment(EVulkanShaderVersion ShaderVersion)
 {
 	// All desktop Vulkan platforms create a Vulkan 1.1 instance (SP_VULKAN_SM5_ANDROID can still use 1.0)
-	return (ShaderPlatform == SP_VULKAN_SM5) ? 
+	return (ShaderVersion == EVulkanShaderVersion::SM5) ?
 		CrossCompiler::FShaderConductorOptions::ETargetEnvironment::Vulkan_1_1: 
 		CrossCompiler::FShaderConductorOptions::ETargetEnvironment::Vulkan_1_0;
 }
@@ -2279,7 +2279,8 @@ static bool CompileWithShaderConductor(
 	const FCompilerInfo&	CompilerInfo,
 	FShaderCompilerOutput&	Output,
 	FVulkanBindingTable&	BindingTable,
-	bool					bStripReflect)
+	bool					bStripReflect,
+	CrossCompiler::FShaderConductorOptions::ETargetEnvironment MinTargetEnvironment)
 {
 	const FShaderCompilerInput& Input = CompilerInfo.Input;
 
@@ -2312,7 +2313,7 @@ static bool CompileWithShaderConductor(
 	}
 	else
 	{
-		Options.TargetEnvironment = GetMinimumTargetEnvironment(Input.Target.GetPlatform());
+		Options.TargetEnvironment = MinTargetEnvironment;
 	}
 
 	UE::ShaderCompilerCommon::DumpDebugShaderData(Input, PreprocessedShader, { CompilerInfo.CCFlags });
@@ -2471,8 +2472,8 @@ void DoCompileVulkanShader(const FShaderCompilerInput& Input, FShaderCompilerOut
 {
 	check(IsVulkanShaderFormat(Input.ShaderFormat));
 
-	const bool bIsSM5 = (Version == EVulkanShaderVersion::SM5);
-	const bool bIsMobile = (Version == EVulkanShaderVersion::ES3_1 || Version == EVulkanShaderVersion::ES3_1_ANDROID);
+	const bool bIsSM5 = (Version == EVulkanShaderVersion::SM5) || (Version == EVulkanShaderVersion::SM5_ANDROID);
+	const bool bIsMobileES31 = (Version == EVulkanShaderVersion::ES3_1 || Version == EVulkanShaderVersion::ES3_1_ANDROID);
 	bool bStripReflect = Input.IsRayTracingShader();
 	// By default we strip reflecion information for Android platform to avoid issues with older drivers
 	if (IsAndroidShaderFormat(Input.ShaderFormat))
@@ -2481,7 +2482,7 @@ void DoCompileVulkanShader(const FShaderCompilerInput& Input, FShaderCompilerOut
 		bStripReflect = !(StripReflect_Android && *StripReflect_Android == TEXT("0"));
 	}
 
-	const CrossCompiler::FShaderConductorOptions::ETargetEnvironment TargetEnvironment = GetMinimumTargetEnvironment(Input.Target.GetPlatform());
+	const CrossCompiler::FShaderConductorOptions::ETargetEnvironment MinTargetEnvironment = GetMinimumTargetEnvironment(Version);
 
 	const EHlslShaderFrequency FrequencyTable[] =
 	{
@@ -2514,7 +2515,7 @@ void DoCompileVulkanShader(const FShaderCompilerInput& Input, FShaderCompilerOut
 	FShaderCompilerDefinitions AdditionalDefines;
 	AdditionalDefines.SetDefine(TEXT("COMPILER_HLSLCC"), 1);
 	AdditionalDefines.SetDefine(TEXT("COMPILER_VULKAN"), 1);
-	if (bIsMobile)
+	if (bIsMobileES31)
 	{
 		AdditionalDefines.SetDefine(TEXT("ES3_1_PROFILE"), 1);
 		AdditionalDefines.SetDefine(TEXT("VULKAN_PROFILE"), 1);
@@ -2539,7 +2540,7 @@ void DoCompileVulkanShader(const FShaderCompilerInput& Input, FShaderCompilerOut
 		AdditionalDefines.SetDefine(TEXT("PLATFORM_SUPPORTS_INLINE_RAY_TRACING"), 1);
 	}
 
-	if (TargetEnvironment >= CrossCompiler::FShaderConductorOptions::ETargetEnvironment::Vulkan_1_1)
+	if (MinTargetEnvironment >= CrossCompiler::FShaderConductorOptions::ETargetEnvironment::Vulkan_1_1)
 	{
 		AdditionalDefines.SetDefine(TEXT("PLATFORM_SUPPORTS_SM6_0_WAVE_OPERATIONS"), 1);
 		AdditionalDefines.SetDefine(TEXT("VULKAN_SUPPORTS_SUBGROUP_SIZE_CONTROL"), 1);
@@ -2641,10 +2642,10 @@ void DoCompileVulkanShader(const FShaderCompilerInput& Input, FShaderCompilerOut
 
 #if PLATFORM_MAC || PLATFORM_WINDOWS || PLATFORM_LINUX
 	// Cross-compile shader via ShaderConductor (DXC, SPIRV-Tools, SPIRV-Cross)
-	bSuccess = CompileWithShaderConductor(PreprocessedShaderSource, EntryPointName, Frequency, CompilerInfo, Output, BindingTable, bStripReflect);
+	bSuccess = CompileWithShaderConductor(PreprocessedShaderSource, EntryPointName, Frequency, CompilerInfo, Output, BindingTable, bStripReflect, MinTargetEnvironment);
 #endif // PLATFORM_MAC || PLATFORM_WINDOWS || PLATFORM_LINUX
 	
-	ShaderParameterParser.ValidateShaderParameterTypes(Input, bIsMobile, Output);
+	ShaderParameterParser.ValidateShaderParameterTypes(Input, bIsMobileES31, Output);
 	
 	if (bDirectCompile)
 	{
diff --git a/Engine/Source/Developer/VulkanShaderFormat/Private/VulkanShaderFormat.cpp b/Engine/Source/Developer/VulkanShaderFormat/Private/VulkanShaderFormat.cpp
--- a/Engine/Source/Developer/VulkanShaderFormat/Private/VulkanShaderFormat.cpp
+++ b/Engine/Source/Developer/VulkanShaderFormat/Private/VulkanShaderFormat.cpp
@@ -86,7 +86,11 @@ public:
 		{
 			DoCompileVulkanShader(Input, Output, WorkingDirectory, EVulkanShaderVersion::ES3_1_ANDROID);
 		}
-		else if (Format == NAME_VULKAN_SM5 || Format == NAME_VULKAN_SM5_ANDROID)
+		else if (Format == NAME_VULKAN_SM5_ANDROID)
+		{
+			DoCompileVulkanShader(Input, Output, WorkingDirectory, EVulkanShaderVersion::SM5_ANDROID);
+		}
+		else if (Format == NAME_VULKAN_SM5)
 		{
 			DoCompileVulkanShader(Input, Output, WorkingDirectory, EVulkanShaderVersion::SM5);
 		}
diff --git a/Engine/Source/Developer/VulkanShaderFormat/Private/VulkanShaderFormat.h b/Engine/Source/Developer/VulkanShaderFormat/Private/VulkanShaderFormat.h
--- a/Engine/Source/Developer/VulkanShaderFormat/Private/VulkanShaderFormat.h
+++ b/Engine/Source/Developer/VulkanShaderFormat/Private/VulkanShaderFormat.h
@@ -17,6 +17,7 @@ enum class EVulkanShaderVersion
 	ES3_1,
 	ES3_1_ANDROID,
 	SM5,
+	SM5_ANDROID,
 };
 
 extern void DoCompileVulkanShader(const struct FShaderCompilerInput& Input,struct FShaderCompilerOutput& Output,const class FString& WorkingDirectory, EVulkanShaderVersion Version);
