VulkanRHI: Genereate bindless sampler declarations the same way as resources. Vulkan Bindless WIP.

#rnx

diff --git a/Engine/Shaders/Public/BindlessResources.ush b/Engine/Shaders/Public/BindlessResources.ush
--- a/Engine/Shaders/Public/BindlessResources.ush
+++ b/Engine/Shaders/Public/BindlessResources.ush
@@ -66,7 +66,9 @@
 	#define AUTO_BINDLESS_SAMPLER_INDEX(Name)          uint BindlessSampler_##Name
 	#define AUTO_BINDLESS_SAMPLER_VARIABLE(Type, Name) Type GetBindlessSampler##Name() { return GetSamplerFromHeap(Type, BindlessSampler_##Name); }; static const Type Name = GetBindlessSampler##Name()
 
-	#define UB_RESOURCE_MEMBER_SAMPLER(Type, Name) AUTO_BINDLESS_SAMPLER_VARIABLE(Type, Name)
+	#ifndef UB_RESOURCE_MEMBER_SAMPLER  // allow platforms to override ub member sampler decl
+		#define UB_RESOURCE_MEMBER_SAMPLER(Type, Name) AUTO_BINDLESS_SAMPLER_VARIABLE(Type, Name)
+	#endif
 #else
 	#define DEFINE_SAMPLER_VARIABLE(Type, Name)  Type Name
 	#define DEFINE_SAMPLER(Type, Name)           Type Name
diff --git a/Engine/Shaders/Public/Platform/Vulkan/VulkanCommon.ush b/Engine/Shaders/Public/Platform/Vulkan/VulkanCommon.ush
--- a/Engine/Shaders/Public/Platform/Vulkan/VulkanCommon.ush
+++ b/Engine/Shaders/Public/Platform/Vulkan/VulkanCommon.ush
@@ -49,9 +49,10 @@ void ImageInterlockedMaxUInt64(RWTexture2D<UlongType> Dest, uint2 Coord, UlongTy
 	#define UB_RESOURCE_MEMBER_RESOURCE(Type, Name) typedef Type SafeType##Name; SafeType##Name VULKAN_RESOURCE_HEAP(SafeType##Name)[]; static const Type Name = VULKAN_RESOURCE_HEAP(SafeType##Name)[BindlessResource_##Name]
 
 	// bindless sampler heaps (specify a register so AUTO_BINDLESS macros aren't applied on them)
-	SamplerState VULKAN_SAMPLER_HEAP(SamplerState)[] : register(s0, space1000);
-	SamplerComparisonState VULKAN_SAMPLER_HEAP(SamplerComparisonState)[] : register(s0, space1000);
-	COMPILER_DEFINE GetSamplerFromHeap(Type, Index) VULKAN_SAMPLER_HEAP(Type)[Index]
+	COMPILER_DEFINE GetSamplerFromHeap(Type, Index)  /* filled by compiler */
+
+	// Override the UB_RESOURCE_MEMBER_SAMPLER to use the Vulkan sampler heaps
+	#define UB_RESOURCE_MEMBER_SAMPLER(Type, Name) typedef Type SafeType##Name; SafeType##Name VULKAN_SAMPLER_HEAP(SafeType##Name)[]; static const Type Name = VULKAN_SAMPLER_HEAP(SafeType##Name)[BindlessSampler_##Name]
 #endif
 
 #if VULKAN_SUPPORTS_SUBGROUP_SIZE_CONTROL
diff --git a/Engine/Source/Developer/VulkanShaderFormat/Private/VulkanShaderCompiler.cpp b/Engine/Source/Developer/VulkanShaderFormat/Private/VulkanShaderCompiler.cpp
--- a/Engine/Source/Developer/VulkanShaderFormat/Private/VulkanShaderCompiler.cpp
+++ b/Engine/Source/Developer/VulkanShaderFormat/Private/VulkanShaderCompiler.cpp
@@ -2422,15 +2422,19 @@ static bool CompileWithShaderConductor(
 class FVulkanShaderParameterParser : public FShaderParameterParser
 {
 public:
-	FVulkanShaderParameterParser(bool bHasBindlessResources)
-		: bEnabled(bHasBindlessResources)
+	FVulkanShaderParameterParser(bool bHasBindlessSamplers, bool bHasBindlessResources)
+		: bEnabled(bHasBindlessSamplers || bHasBindlessResources)
 	{}
 
 protected:
 	FString GenerateBindlessParameterDeclaration(const FParsedShaderParameter& ParsedParameter) const override
 	{
-		if (bEnabled && (ParsedParameter.BindlessConversionType == EBindlessConversionType::Resource))
+		if (bEnabled)
 		{
+			const bool IsSampler = (ParsedParameter.BindlessConversionType == EBindlessConversionType::Sampler);
+			const TCHAR* IndexPrefix = IsSampler ? TEXT("BindlessSampler_") : TEXT("BindlessResource_");
+			const TCHAR* HeapPrefix = IsSampler ? VulkanBindless::kBindlessSamplerArrayPrefix : VulkanBindless::kBindlessResourceArrayPrefix;
+
 			const TCHAR* StorageClass = ParsedParameter.bGloballyCoherent ? TEXT("globallycoherent ") : TEXT("");
 
 			const FStringView Name = ParsedParameter.ParsedName;
@@ -2444,7 +2448,7 @@ protected:
 			if (ParsedParameter.ConstantBufferParameterType == EShaderParameterType::Num)
 			{
 				// e.g. `uint BindlessResource_##Name;`
-				Result << TEXT("uint BindlessResource_") << Name << TEXT("; ");
+				Result << TEXT("uint ") << IndexPrefix << Name << TEXT("; ");
 			}
 
 			// Add the typedef
@@ -2452,15 +2456,19 @@ protected:
 
 			// Declare a heap for the RewriteType
 			// e.g. `SafeType##Name ResourceDescriptorHeap_SafeType##Name[];`
-			Result << RewriteType << TEXT(" ") << VulkanBindless::kBindlessResourceArrayPrefix << RewriteType << TEXT("[]; ");
+			Result << RewriteType << TEXT(" ") << HeapPrefix << RewriteType << TEXT("[]; ");
 			// :todo-jn: specify the descripor set and binding directly in source instead of patching SPIRV
 
 			// e.g. `static const Type Name = GetBindlessResource##Name()`
-			Result << TEXT("static const ") << StorageClass << RewriteType << TEXT(" ") << Name << TEXT(" = ") << VulkanBindless::kBindlessResourceArrayPrefix << RewriteType << TEXT("[BindlessResource_") << Name << TEXT("];");
+			Result << TEXT("static const ") << StorageClass << RewriteType << TEXT(" ") << Name << TEXT(" = ") << HeapPrefix << RewriteType << TEXT("[") << IndexPrefix << Name << TEXT("];");
 
 			return Result.ToString();
 		}
-		return FShaderParameterParser::GenerateBindlessParameterDeclaration(ParsedParameter);
+		else
+		{
+			// use original code path
+			return FShaderParameterParser::GenerateBindlessParameterDeclaration(ParsedParameter);
+		}
 	}
 
 private:
@@ -2577,8 +2585,9 @@ void DoCompileVulkanShader(const FShaderCompilerInput& Input, FShaderCompilerOut
 		}
 	}
 
+	const bool bHasBindlessSamplers = Input.Environment.CompilerFlags.Contains(CFLAG_BindlessSamplers);
 	const bool bHasBindlessResources = Input.Environment.CompilerFlags.Contains(CFLAG_BindlessResources);
-	FVulkanShaderParameterParser ShaderParameterParser(bHasBindlessResources);
+	FVulkanShaderParameterParser ShaderParameterParser(bHasBindlessSamplers, bHasBindlessResources);
 	if (!ShaderParameterParser.ParseAndModify(Input, Output, PreprocessedShaderSource))
 	{
 		// The FShaderParameterParser will add any relevant errors.
