Remove redundant dumping of .spv file in VulkanShaderCompiler.

#rnx

diff --git a/Engine/Source/Developer/VulkanShaderFormat/Private/VulkanShaderCompiler.cpp b/Engine/Source/Developer/VulkanShaderFormat/Private/VulkanShaderCompiler.cpp
--- a/Engine/Source/Developer/VulkanShaderFormat/Private/VulkanShaderCompiler.cpp
+++ b/Engine/Source/Developer/VulkanShaderFormat/Private/VulkanShaderCompiler.cpp
@@ -2411,11 +2411,6 @@ static bool CompileWithShaderConductor(
 		Output.ShaderCode.AddOptionalData(FShaderCodeName::Key, TCHAR_TO_UTF8(*Input.GenerateShaderName()));
 	}
 
-	if (bDebugDump)
-	{
-		DumpDebugShaderBinary(Input, Spirv.GetByteData(), Spirv.GetByteSize(), TEXT("spv"));
-	}
-
 	// Flush warnings
 	CompilerContext.FlushErrors(Output.Errors);
 	return true;
