Fixing bad vulkan.h include.
Switching back to referencing VulkanRHI headers correctly.

diff --git a/Engine/Source/Developer/VulkanShaderFormat/Private/VulkanShaderCompiler.cpp b/Engine/Source/Developer/VulkanShaderFormat/Private/VulkanShaderCompiler.cpp
--- a/Engine/Source/Developer/VulkanShaderFormat/Private/VulkanShaderCompiler.cpp
+++ b/Engine/Source/Developer/VulkanShaderFormat/Private/VulkanShaderCompiler.cpp
@@ -36,7 +36,7 @@ typedef enum VkDescriptorType {
     VK_DESCRIPTOR_TYPE_MAX_ENUM = 0x7FFFFFFF
 } VkDescriptorType;
 #else
-#include "vulkan.h"
+#include "IVulkanDynamicRHI.h"
 #endif
 #include "VulkanBackend.h"
 #include "VulkanShaderResources.h"
diff --git a/Engine/Source/Developer/VulkanShaderFormat/VulkanShaderFormat.Build.cs b/Engine/Source/Developer/VulkanShaderFormat/VulkanShaderFormat.Build.cs
--- a/Engine/Source/Developer/VulkanShaderFormat/VulkanShaderFormat.Build.cs
+++ b/Engine/Source/Developer/VulkanShaderFormat/VulkanShaderFormat.Build.cs
@@ -1,7 +1,6 @@
 // Copyright Epic Games, Inc. All Rights Reserved.
 
 using UnrealBuildTool;
-using System.IO;
 
 public class VulkanShaderFormat : ModuleRules
 {
@@ -10,7 +9,7 @@ public class VulkanShaderFormat : ModuleRules
 		PrivateIncludePathModuleNames.Add("TargetPlatform");
 
 		// Do not link the module (as that would require the vulkan dll), only the include paths
-		PublicIncludePaths.Add("Runtime/VulkanRHI/Public");
+		PublicIncludePathModuleNames.Add("VulkanRHI");
 
 		PrivateDependencyModuleNames.AddRange(
 			new string[] {
