VulkanRHI: Add Vulkan_1_3 support for shader compilation. (WIP towards VULKAN_SM6)

#rnx

diff --git a/Engine/Source/Developer/ShaderCompilerCommon/Private/ShaderConductorContext.cpp b/Engine/Source/Developer/ShaderCompilerCommon/Private/ShaderConductorContext.cpp
--- a/Engine/Source/Developer/ShaderCompilerCommon/Private/ShaderConductorContext.cpp
+++ b/Engine/Source/Developer/ShaderCompilerCommon/Private/ShaderConductorContext.cpp
@@ -390,6 +390,9 @@ namespace CrossCompiler
 		case ETargetEnvironment::Vulkan_1_2:
 			DxcArgRefs.Add("-fspv-target-env=vulkan1.2");
 			break;
+		case ETargetEnvironment::Vulkan_1_3:
+			DxcArgRefs.Add("-fspv-target-env=vulkan1.3");
+			break;
 		}
 
 		if (DxcArgRefs.Num() > 0)
diff --git a/Engine/Source/Developer/ShaderCompilerCommon/Public/ShaderConductorContext.h b/Engine/Source/Developer/ShaderCompilerCommon/Public/ShaderConductorContext.h
--- a/Engine/Source/Developer/ShaderCompilerCommon/Public/ShaderConductorContext.h
+++ b/Engine/Source/Developer/ShaderCompilerCommon/Public/ShaderConductorContext.h
@@ -96,8 +96,9 @@ namespace CrossCompiler
 			Vulkan_1_0,
 			Vulkan_1_1,
 			Vulkan_1_2,
+			Vulkan_1_3,
 		};
-		ETargetEnvironment TargetEnvironment = ETargetEnvironment::Vulkan_1_0;
+		ETargetEnvironment TargetEnvironment = ETargetEnvironment::Vulkan_1_1;
 
 		/** Shader model version of the input language. By default SM6.2. */
 		FHlslShaderModel ShaderModel = { 6, 2 };
diff --git a/Engine/Source/Developer/VulkanShaderFormat/Private/VulkanShaderCompiler.cpp b/Engine/Source/Developer/VulkanShaderFormat/Private/VulkanShaderCompiler.cpp
--- a/Engine/Source/Developer/VulkanShaderFormat/Private/VulkanShaderCompiler.cpp
+++ b/Engine/Source/Developer/VulkanShaderFormat/Private/VulkanShaderCompiler.cpp
@@ -67,7 +67,9 @@ inline bool SupportsOfflineCompiler(FName ShaderFormat)
 
 inline CrossCompiler::FShaderConductorOptions::ETargetEnvironment GetMinimumTargetEnvironment(EVulkanShaderVersion ShaderVersion)
 {
-	return CrossCompiler::FShaderConductorOptions::ETargetEnvironment::Vulkan_1_1;
+	return (ShaderVersion == EVulkanShaderVersion::SM6) ?
+		CrossCompiler::FShaderConductorOptions::ETargetEnvironment::Vulkan_1_3:
+		CrossCompiler::FShaderConductorOptions::ETargetEnvironment::Vulkan_1_1;
 }
 
 DEFINE_LOG_CATEGORY_STATIC(LogVulkanShaderCompiler, Log, All); 
@@ -2197,6 +2199,10 @@ static void VulkanCreateDXCCompileBatchFiles(
 	{
 		VulkanVersion = TEXT("vulkan1.2");
 	}
+	else if (Options.TargetEnvironment == CrossCompiler::FShaderConductorOptions::ETargetEnvironment::Vulkan_1_3)
+	{
+		VulkanVersion = TEXT("vulkan1.3");
+	}
 	else
 	{
 		ensure(false);
diff --git a/Engine/Source/Developer/VulkanShaderFormat/Private/VulkanShaderFormat.h b/Engine/Source/Developer/VulkanShaderFormat/Private/VulkanShaderFormat.h
--- a/Engine/Source/Developer/VulkanShaderFormat/Private/VulkanShaderFormat.h
+++ b/Engine/Source/Developer/VulkanShaderFormat/Private/VulkanShaderFormat.h
@@ -18,6 +18,7 @@ enum class EVulkanShaderVersion
 	ES3_1_ANDROID,
 	SM5,
 	SM5_ANDROID,
+	SM6,
 };
 
 extern void DoCompileVulkanShader(const struct FShaderCompilerInput& Input,struct FShaderCompilerOutput& Output,const class FString& WorkingDirectory, EVulkanShaderVersion Version);
