VulkanRHI: Add VULKAN_PROFILE_SM6 definition to shaders for upcoming VULKAN_SM6.

#rnx

diff --git a/Engine/Plugins/FX/Niagara/Shaders/Private/NiagaraEmitterInstanceShader.usf b/Engine/Plugins/FX/Niagara/Shaders/Private/NiagaraEmitterInstanceShader.usf
--- a/Engine/Plugins/FX/Niagara/Shaders/Private/NiagaraEmitterInstanceShader.usf
+++ b/Engine/Plugins/FX/Niagara/Shaders/Private/NiagaraEmitterInstanceShader.usf
@@ -577,7 +577,7 @@ int rand_int(int x, int Seed1, int Seed2, int Seed3)
 		RWIDToIndexTable[IDIndex] = InstanceIndex;
 	}
 
-#define USE_GROUP_SHARED ((THREADGROUP_SIZE == 64 || THREADGROUP_SIZE == 32) && !USE_WAVE_INTRINSICS && !VULKAN_PROFILE_SM5 && !VULKAN_PROFILE)
+#define USE_GROUP_SHARED ((THREADGROUP_SIZE == 64 || THREADGROUP_SIZE == 32) && !USE_WAVE_INTRINSICS && !VULKAN_PROFILE_SM5 && !VULKAN_PROFILE_SM6 && !VULKAN_PROFILE)
 
 #if USE_GROUP_SHARED
 	#if THREADGROUP_SIZE == 64
diff --git a/Engine/Shaders/Private/BasePassPixelShader.usf b/Engine/Shaders/Private/BasePassPixelShader.usf
--- a/Engine/Shaders/Private/BasePassPixelShader.usf
+++ b/Engine/Shaders/Private/BasePassPixelShader.usf
@@ -1366,8 +1366,8 @@ void FPixelShaderInOut_MainPS(
 	float3 OutOfBoundsMaskLuminance = 0;
 	// This feature is only needed for development/editor - we can compile it out for a shipping build (see r.CompileShadersForDevelopment cvar help)
 	// The following block is disabled on Vulkan because it triggers an nvidia driver bug (UE-101609).
-	#if !VULKAN_PROFILE_SM5
-		#if METAL_SM5_PROFILE || SM6_PROFILE || SM5_PROFILE || VULKAN_PROFILE_SM5
+	#if !VULKAN_PROFILE_SM5 && !VULKAN_PROFILE_SM6
+		#if METAL_SM5_PROFILE || SM6_PROFILE || SM5_PROFILE || VULKAN_PROFILE_SM5 || VULKAN_PROFILE_SM6
 		BRANCH
 		if (View.OutOfBoundsMask > 0)
 		{
diff --git a/Engine/Shaders/Public/Platform.ush b/Engine/Shaders/Public/Platform.ush
--- a/Engine/Shaders/Public/Platform.ush
+++ b/Engine/Shaders/Public/Platform.ush
@@ -176,6 +176,10 @@
 #define	VULKAN_PROFILE_SM5 0
 #endif
 
+#ifndef VULKAN_PROFILE_SM6
+#define	VULKAN_PROFILE_SM6 0
+#endif
+
 #ifndef IOS
 #define IOS 0
 #endif
@@ -365,6 +369,10 @@ float2 UnpackFloat2FromUInt(uint In)
 	#define FEATURE_LEVEL FEATURE_LEVEL_SM5
 	#define STENCIL_COMPONENT_SWIZZLE .x
 
+#elif VULKAN_PROFILE_SM6
+	#define FEATURE_LEVEL FEATURE_LEVEL_SM6
+	#define STENCIL_COMPONENT_SWIZZLE .x
+
 #elif METAL_PROFILE
 	#define FEATURE_LEVEL  FEATURE_LEVEL_ES3_1
 	// @todo metal: remove this and make sure all uses handle METAL_PROFILE
@@ -475,7 +483,7 @@ float2 UnpackFloat2FromUInt(uint In)
 // Ideally this would just be enabled for all SM4 and above feature level platforms
 // @todo metal mrt: No reason this can't work with Metal, once cross compiler is fixed
 #ifndef SUPPORTS_INDEPENDENT_SAMPLERS
-#define SUPPORTS_INDEPENDENT_SAMPLERS (SM6_PROFILE || SM5_PROFILE || METAL_MRT_PROFILE || METAL_SM5_PROFILE || VULKAN_PROFILE_SM5 || VULKAN_PROFILE)
+#define SUPPORTS_INDEPENDENT_SAMPLERS (SM6_PROFILE || SM5_PROFILE || METAL_MRT_PROFILE || METAL_SM5_PROFILE || VULKAN_PROFILE_SM5 || VULKAN_PROFILE_SM6 || VULKAN_PROFILE)
 #endif
 
 // Whether the platform support pixel coverage on MSAA targets (SV_Coverage).
@@ -632,7 +640,7 @@ void ShaderYield()
 
 #endif // SM6_PROFILE || SM5_PROFILE || COMPILER_SUPPORTS_ATTRIBUTES
 
-#if SM6_PROFILE || SM5_PROFILE || METAL_MRT_PROFILE || METAL_SM5_PROFILE || ES3_1_PROFILE || VULKAN_PROFILE_SM5
+#if SM6_PROFILE || SM5_PROFILE || METAL_MRT_PROFILE || METAL_SM5_PROFILE || ES3_1_PROFILE || VULKAN_PROFILE_SM5 || VULKAN_PROFILE_SM6
 #define EARLYDEPTHSTENCIL [earlydepthstencil]
 #endif
 
diff --git a/Engine/Source/Developer/VulkanShaderFormat/Private/VulkanShaderCompiler.cpp b/Engine/Source/Developer/VulkanShaderFormat/Private/VulkanShaderCompiler.cpp
--- a/Engine/Source/Developer/VulkanShaderFormat/Private/VulkanShaderCompiler.cpp
+++ b/Engine/Source/Developer/VulkanShaderFormat/Private/VulkanShaderCompiler.cpp
@@ -49,6 +49,7 @@ inline bool IsVulkanShaderFormat(FName ShaderFormat)
 	return ShaderFormat == NAME_VULKAN_ES3_1_ANDROID
 		|| ShaderFormat == NAME_VULKAN_ES3_1
 		|| ShaderFormat == NAME_VULKAN_SM5
+//		|| ShaderFormat == NAME_VULKAN_SM6 // :todo-jn:
 		|| ShaderFormat == NAME_VULKAN_SM5_ANDROID;
 }
 
@@ -2548,6 +2549,7 @@ void DoCompileVulkanShader(const FShaderCompilerInput& Input, FShaderCompilerOut
 {
 	check(IsVulkanShaderFormat(Input.ShaderFormat));
 
+	const bool bIsSM6 = (Version == EVulkanShaderVersion::SM6);
 	const bool bIsSM5 = (Version == EVulkanShaderVersion::SM5) || (Version == EVulkanShaderVersion::SM5_ANDROID);
 	const bool bIsMobileES31 = (Version == EVulkanShaderVersion::ES3_1 || Version == EVulkanShaderVersion::ES3_1_ANDROID);
 	bool bStripReflect = Input.IsRayTracingShader();
@@ -2566,12 +2568,12 @@ void DoCompileVulkanShader(const FShaderCompilerInput& Input, FShaderCompilerOut
 		HSF_InvalidFrequency,
 		HSF_InvalidFrequency,
 		HSF_PixelShader,
-		bIsSM5 ? HSF_GeometryShader : HSF_InvalidFrequency,
+		(bIsSM5 || bIsSM6) ? HSF_GeometryShader : HSF_InvalidFrequency,
 		HSF_ComputeShader, 
-		bIsSM5 ? HSF_RayGen : HSF_InvalidFrequency,
-		bIsSM5 ? HSF_RayMiss : HSF_InvalidFrequency,
-		bIsSM5 ? HSF_RayHitGroup : HSF_InvalidFrequency,
-		bIsSM5 ? HSF_RayCallable : HSF_InvalidFrequency,
+		(bIsSM5 || bIsSM6) ? HSF_RayGen : HSF_InvalidFrequency,
+		(bIsSM5 || bIsSM6) ? HSF_RayMiss : HSF_InvalidFrequency,
+		(bIsSM5 || bIsSM6) ? HSF_RayHitGroup : HSF_InvalidFrequency,
+		(bIsSM5 || bIsSM6) ? HSF_RayCallable : HSF_InvalidFrequency,
 	};
 
 	const EShaderFrequency Frequency = (EShaderFrequency)Input.Target.Frequency;
@@ -2596,6 +2598,10 @@ void DoCompileVulkanShader(const FShaderCompilerInput& Input, FShaderCompilerOut
 		AdditionalDefines.SetDefine(TEXT("ES3_1_PROFILE"), 1);
 		AdditionalDefines.SetDefine(TEXT("VULKAN_PROFILE"), 1);
 	}
+	else if (bIsSM6)
+	{
+		AdditionalDefines.SetDefine(TEXT("VULKAN_PROFILE_SM6"), 1);
+	}
 	else if (bIsSM5)
 	{
 		AdditionalDefines.SetDefine(TEXT("VULKAN_PROFILE_SM5"), 1);
