VulkanRHI: Fix Vulkan shader compilation by enabling VK_EXT_scalar_block_layout in all shaders for VULKAN_SM6 (its presence is verified by the profile).

#rnx

diff --git a/Engine/Source/Developer/VulkanShaderFormat/Private/VulkanShaderCompiler.cpp b/Engine/Source/Developer/VulkanShaderFormat/Private/VulkanShaderCompiler.cpp
--- a/Engine/Source/Developer/VulkanShaderFormat/Private/VulkanShaderCompiler.cpp
+++ b/Engine/Source/Developer/VulkanShaderFormat/Private/VulkanShaderCompiler.cpp
@@ -2357,7 +2357,10 @@ static bool CompileWithShaderConductor(
 
 	// Initialize compilation options for ShaderConductor
 	CrossCompiler::FShaderConductorOptions Options;
-	Options.bDisableScalarBlockLayout = !bIsRayTracingShader;
+
+	// VK_EXT_scalar_block_layout is required by raytracing and by Nanite (so expect it to be present in SM6/Vulkan_1_3)
+	Options.bDisableScalarBlockLayout = !(bIsRayTracingShader || 
+		MinTargetEnvironment >= CrossCompiler::FShaderConductorOptions::ETargetEnvironment::Vulkan_1_3);
 
 	// Enable HLSL 2021 if specified
 	if (Input.Environment.CompilerFlags.Contains(CFLAG_HLSL2021))
