VulkanRHI:  VULKAN_SM6 fix for some shaders getting compiled for Vulkan 1.2 even when MinTargetEnvironment is set to 1.3.

#rnx

diff --git a/Engine/Source/Developer/VulkanShaderFormat/Private/VulkanShaderCompiler.cpp b/Engine/Source/Developer/VulkanShaderFormat/Private/VulkanShaderCompiler.cpp
--- a/Engine/Source/Developer/VulkanShaderFormat/Private/VulkanShaderCompiler.cpp
+++ b/Engine/Source/Developer/VulkanShaderFormat/Private/VulkanShaderCompiler.cpp
@@ -2357,6 +2357,7 @@ static bool CompileWithShaderConductor(
 
 	// Initialize compilation options for ShaderConductor
 	CrossCompiler::FShaderConductorOptions Options;
+	Options.TargetEnvironment = MinTargetEnvironment;
 
 	// VK_EXT_scalar_block_layout is required by raytracing and by Nanite (so expect it to be present in SM6/Vulkan_1_3)
 	Options.bDisableScalarBlockLayout = !(bIsRayTracingShader || 
@@ -2376,10 +2377,6 @@ static bool CompileWithShaderConductor(
 			Options.TargetEnvironment = CrossCompiler::FShaderConductorOptions::ETargetEnvironment::Vulkan_1_2;
 		}
 	}
-	else
-	{
-		Options.TargetEnvironment = MinTargetEnvironment;
-	}
 
 	UE::ShaderCompilerCommon::DumpDebugShaderData(Input, PreprocessedShader, { CompilerInfo.CCFlags });
 
