VulkanRHI: VKRT WIP: Increase size of bitfield to track uniform buffer usage.  Fix copy-paste error where acceleration structures were sent to wrong bindless descriptor set.

#rnx

diff --git a/Engine/Source/Developer/VulkanShaderFormat/Private/VulkanShaderCompiler.cpp b/Engine/Source/Developer/VulkanShaderFormat/Private/VulkanShaderCompiler.cpp
--- a/Engine/Source/Developer/VulkanShaderFormat/Private/VulkanShaderCompiler.cpp
+++ b/Engine/Source/Developer/VulkanShaderFormat/Private/VulkanShaderCompiler.cpp
@@ -969,10 +969,11 @@ static void BuildShaderOutput(
 
 	FOLDVulkanCodeHeader OLDHeader;
 
-	const EShaderFrequency Frequency = (EShaderFrequency)ShaderOutput.Target.Frequency;
+	const EShaderFrequency Frequency = static_cast<EShaderFrequency>(ShaderOutput.Target.Frequency);
 
 	TBitArray<> UsedUniformBufferSlots;
-	UsedUniformBufferSlots.Init(false, 32);
+	const int32 MaxNumBits = VulkanBindless::MaxUniformBuffersPerStage * SF_NumFrequencies;
+	UsedUniformBufferSlots.Init(false, MaxNumBits);
 
 
 	static const FString AttributePrefix = TEXT("in_ATTRIBUTE");
@@ -1015,7 +1016,6 @@ static void BuildShaderOutput(
 	for (auto& UniformBlock : CCHeader.UniformBlocks)
 	{
 		// DXC's generated "$Globals" has been converted to "_Globals" at this point
-		uint16 UBIndex = UniformBlock.Index;
 		if (UniformBlock.Name.StartsWith(TEXT("HLSLCC_CB")) || UniformBlock.Name.StartsWith(TEXT("_Globals")))
 		{
 			// Skip...
@@ -1023,7 +1023,7 @@ static void BuildShaderOutput(
 		else
 		{
 			// Regular UB
-			int32 VulkanBindingIndex = Spirv.FindBinding(UniformBlock.Name, true);
+			const int32 VulkanBindingIndex = Spirv.FindBinding(UniformBlock.Name, true);
 			check(VulkanBindingIndex != -1);
 			check(!UsedUniformBufferSlots[VulkanBindingIndex]);
 			UsedUniformBufferSlots[VulkanBindingIndex] = true;
@@ -1037,7 +1037,7 @@ static void BuildShaderOutput(
 
 
 	const bool bSupportsBindless = ShaderInput.Environment.CompilerFlags.Contains(CFLAG_BindlessResources) || ShaderInput.Environment.CompilerFlags.Contains(CFLAG_BindlessSamplers);
-	const int32 StageOffset = bSupportsBindless ? (ShaderStage::GetStageForFrequency(static_cast<EShaderFrequency>(ShaderInput.Target.Frequency)) * VulkanBindless::MaxUniformBuffersPerStage) : 0;
+	const int32 StageOffset = bSupportsBindless ? (ShaderStage::GetStageForFrequency(Frequency) * VulkanBindless::MaxUniformBuffersPerStage) : 0;
 
 	const TArray<FVulkanBindingTable::FBinding>& HlslccBindings = BindingTable.GetBindings();
 	OLDHeader.NEWDescriptorInfo.NumBufferInfos = 0;
@@ -1499,7 +1499,7 @@ static void GatherSpirvReflectionBindings(
 		MoveBindlessHeaps(OutBindings.TextureUAVs, VulkanBindless::kBindlessResourceArrayPrefix, VulkanBindless::BindlessStorageImageSet);
 		MoveBindlessHeaps(OutBindings.TBufferSRVs, VulkanBindless::kBindlessResourceArrayPrefix, VulkanBindless::BindlessUniformTexelBufferSet);
 		MoveBindlessHeaps(OutBindings.TBufferUAVs, VulkanBindless::kBindlessResourceArrayPrefix, VulkanBindless::BindlessStorageTexelBufferSet);
-		MoveBindlessHeaps(OutBindings.AccelerationStructures, VulkanBindless::kBindlessResourceArrayPrefix, VulkanBindless::BindlessSamplerSet);
+		MoveBindlessHeaps(OutBindings.AccelerationStructures, VulkanBindless::kBindlessResourceArrayPrefix, VulkanBindless::BindlessAccelerationStructureSet);
 
 		// Move uniform buffers to the correct set
 		{
