VulkanRHI: Fix issues with globallycoherent in Vulkan bindless, move it to heap decl for now (properly adds the "OpDecorate Coherent" in SPIRV).

#rnx

diff --git a/Engine/Source/Developer/VulkanShaderFormat/Private/VulkanShaderCompiler.cpp b/Engine/Source/Developer/VulkanShaderFormat/Private/VulkanShaderCompiler.cpp
--- a/Engine/Source/Developer/VulkanShaderFormat/Private/VulkanShaderCompiler.cpp
+++ b/Engine/Source/Developer/VulkanShaderFormat/Private/VulkanShaderCompiler.cpp
@@ -2518,7 +2518,7 @@ protected:
 
 			// Declare a heap for the RewriteType
 			// e.g. `SafeType##Name ResourceDescriptorHeap_SafeType##Name[];`
-			Result << RewriteType << TEXT(" ") << HeapPrefix << RewriteType << TEXT("[]; ");
+			Result << StorageClass << RewriteType << TEXT(" ") << HeapPrefix << RewriteType << TEXT("[]; ");
 			// :todo-jn: specify the descripor set and binding directly in source instead of patching SPIRV
 
 			// e.g. `static const Type Name = GetBindlessResource##Name()`
