[shaders] remove unnecessary passing of hlslcc flags as an option to debug info output. this was just used to pass a switch on the SCW directcompile cmdline which SCW no longer does anything with. also cleaned up some additional unnecessary code related to these flags in backends that no longer use hlslcc at all.

diff --git a/Engine/Source/Developer/Apple/MetalShaderFormat/Private/MetalDerivedData.cpp b/Engine/Source/Developer/Apple/MetalShaderFormat/Private/MetalDerivedData.cpp
--- a/Engine/Source/Developer/Apple/MetalShaderFormat/Private/MetalDerivedData.cpp
+++ b/Engine/Source/Developer/Apple/MetalShaderFormat/Private/MetalDerivedData.cpp
@@ -22,7 +22,6 @@ extern void BuildMetalShaderOutput(
 	FShaderCompilerOutput& ShaderOutput,
 	const FShaderCompilerInput& ShaderInput,
 	FSHAHash const& GUIDHash,
-	uint32 CCFlags,
 	const ANSICHAR* InShaderSource,
 	uint32 SourceLen,
 	uint32 SourceCRCLen,
@@ -199,7 +198,6 @@ bool DoCompileMetalShader(
 	const FString& InPreprocessedShader,
 	FSHAHash GUIDHash,
 	uint32 VersionEnum,
-	uint32 CCFlags,
 	EMetalGPUSemantics Semantics,
 	uint32 MaxUnrollLoops,
 	EShaderFrequency Frequency,
@@ -1412,7 +1410,7 @@ bool DoCompileMetalShader(
 	if (Result != 0)
 	{
 		Output.Target = Input.Target;
-		BuildMetalShaderOutput(Output, Input, GUIDHash, CCFlags, MetalSource.c_str(), MetalSource.length(), CRCLen, CRC, VersionEnum, *Standard, *MinOSVersion, Output.Errors, OutputData.TypedBuffers, OutputData.InvariantBuffers, OutputData.TypedUAVs, OutputData.ConstantBuffers, bAllowFastIntrinsics);
+		BuildMetalShaderOutput(Output, Input, GUIDHash, MetalSource.c_str(), MetalSource.length(), CRCLen, CRC, VersionEnum, *Standard, *MinOSVersion, Output.Errors, OutputData.TypedBuffers, OutputData.InvariantBuffers, OutputData.TypedUAVs, OutputData.ConstantBuffers, bAllowFastIntrinsics);
 		return Output.bSucceeded;
 	}
 	else
diff --git a/Engine/Source/Developer/Apple/MetalShaderFormat/Private/MetalDerivedData.h b/Engine/Source/Developer/Apple/MetalShaderFormat/Private/MetalDerivedData.h
--- a/Engine/Source/Developer/Apple/MetalShaderFormat/Private/MetalDerivedData.h
+++ b/Engine/Source/Developer/Apple/MetalShaderFormat/Private/MetalDerivedData.h
@@ -97,7 +97,6 @@ bool DoCompileMetalShader(
 	const FString& PreprocessedShader,
 	FSHAHash GUIDHash,
 	uint32 VersionEnum,
-	uint32 CCFlags,
 	EMetalGPUSemantics Semantics,
 	uint32 MaxUnrollLoops,
 	EShaderFrequency Frequency,
diff --git a/Engine/Source/Developer/Apple/MetalShaderFormat/Private/MetalShaderCompiler.cpp b/Engine/Source/Developer/Apple/MetalShaderFormat/Private/MetalShaderCompiler.cpp
--- a/Engine/Source/Developer/Apple/MetalShaderFormat/Private/MetalShaderCompiler.cpp
+++ b/Engine/Source/Developer/Apple/MetalShaderFormat/Private/MetalShaderCompiler.cpp
@@ -227,7 +227,6 @@ void BuildMetalShaderOutput(
 	FShaderCompilerOutput& ShaderOutput,
 	const FShaderCompilerInput& ShaderInput,
 	FSHAHash const& GUIDHash,
-	uint32 CCFlags,
 	const ANSICHAR* InShaderSource,
 	uint32 SourceLen,
 	uint32 SourceCRCLen,
@@ -1052,22 +1051,10 @@ void CompileShader_Metal(const FShaderCompilerInput& _Input,FShaderCompilerOutpu
 
 	Output.PreprocessTime = FPlatformTime::Seconds() - StartPreprocessTime;
 	
-	uint32 CCFlags = HLSLCC_NoPreprocess | HLSLCC_PackUniformsIntoUniformBufferWithNames | HLSLCC_FixAtomicReferences | HLSLCC_RetainSizes | HLSLCC_KeepSamplerAndImageNames;
-	if (!bDirectCompile || UE_BUILD_DEBUG)
-	{
-		// Validation is expensive - only do it when compiling directly for debugging
-		CCFlags |= HLSLCC_NoValidation;
-	}
-
-	// Required as we added the RemoveUniformBuffersFromSource() function (the cross-compiler won't be able to interpret comments w/o a preprocessor)
-	CCFlags &= ~HLSLCC_NoPreprocess;
-
 	// Write out the preprocessed file and a batch file to compile it if requested (DumpDebugInfoPath is valid)
 	if (bDumpDebugInfo && !bDirectCompile)
 	{
-		UE::ShaderCompilerCommon::FDebugShaderDataOptions DebugDataOptions;
-		DebugDataOptions.HlslCCFlags = CCFlags;
-		UE::ShaderCompilerCommon::DumpDebugShaderData(Input, PreprocessedShader, DebugDataOptions);
+		UE::ShaderCompilerCommon::DumpDebugShaderData(Input, PreprocessedShader);
 	}
 
 	FSHAHash GUIDHash;
@@ -1084,7 +1071,7 @@ void CompileShader_Metal(const FShaderCompilerInput& _Input,FShaderCompilerOutpu
 		FSHA1::HashBuffer(&Guid, sizeof(FGuid), GUIDHash.Hash);
 	}
 
-	bool bCompiled = DoCompileMetalShader(Input, Output, WorkingDirectory, PreprocessedShader, GUIDHash, VersionEnum, CCFlags, Semantics, MaxUnrollLoops, Frequency, bDumpDebugInfo, Standard, MinOSVersion);
+	bool bCompiled = DoCompileMetalShader(Input, Output, WorkingDirectory, PreprocessedShader, GUIDHash, VersionEnum, Semantics, MaxUnrollLoops, Frequency, bDumpDebugInfo, Standard, MinOSVersion);
 	if (bCompiled && Output.bSucceeded)
 	{
 
diff --git a/Engine/Source/Developer/ShaderCompilerCommon/Private/ShaderCompilerCommon.cpp b/Engine/Source/Developer/ShaderCompilerCommon/Private/ShaderCompilerCommon.cpp
--- a/Engine/Source/Developer/ShaderCompilerCommon/Private/ShaderCompilerCommon.cpp
+++ b/Engine/Source/Developer/ShaderCompilerCommon/Private/ShaderCompilerCommon.cpp
@@ -1201,7 +1201,7 @@ void TransformStringIntoCharacterArray(FString& PreprocessedShaderSource, TArray
 
 ///////////////////////////////////////////////////////////////////////////////////////////////////////////////////////////////////////
 
-FString CreateShaderCompilerWorkerDirectCommandLine(const FShaderCompilerInput& Input, uint32 CCFlags = 0)
+FString CreateShaderCompilerWorkerDirectCommandLine(const FShaderCompilerInput& Input)
 {
 	FString Text(TEXT("-directcompile -format="));
 	Text += Input.ShaderFormat.GetPlainNameString();
@@ -1250,11 +1250,6 @@ FString CreateShaderCompilerWorkerDirectCommandLine(const FShaderCompilerInput&
 	Text += TEXT(" -cflags=");
 	Text += FString::Printf(TEXT("%llu"), Input.Environment.CompilerFlags.GetData());
 
-	if (CCFlags)
-	{
-		Text += TEXT(" -hlslccflags=");
-		Text += FString::Printf(TEXT("%u"), CCFlags);
-	}
 	// When we're running in directcompile mode, we don't to spam the crash reporter
 	Text += TEXT(" -nocrashreports");
 	return Text;
@@ -1605,7 +1600,7 @@ namespace UE::ShaderCompilerCommon
 		Contents += TEXT("\n");
 		Contents += CrossCompiler::CreateResourceTableFromEnvironment(Input.Environment);
 		Contents += TEXT("#if 0 /*DIRECT COMPILE*/\n");
-		Contents += CreateShaderCompilerWorkerDirectCommandLine(Input, Options.HlslCCFlags);
+		Contents += CreateShaderCompilerWorkerDirectCommandLine(Input);
 		Contents += TEXT("\n#endif /*DIRECT COMPILE*/\n");
 		if (!Input.DebugDescription.IsEmpty())
 		{
diff --git a/Engine/Source/Developer/ShaderCompilerCommon/Public/ShaderCompilerCommon.h b/Engine/Source/Developer/ShaderCompilerCommon/Public/ShaderCompilerCommon.h
--- a/Engine/Source/Developer/ShaderCompilerCommon/Public/ShaderCompilerCommon.h
+++ b/Engine/Source/Developer/ShaderCompilerCommon/Public/ShaderCompilerCommon.h
@@ -120,6 +120,7 @@ namespace UE::ShaderCompilerCommon
 			const TCHAR* BaseFileName;
 			const TCHAR* Data;
 		};
+		UE_DEPRECATED(5.4, "HlslCCFlags field is no longer used in debug output; please remove any usage.")
 		uint32 HlslCCFlags = 0;
 		const TCHAR* OverrideBaseFilename = nullptr;
 		const TCHAR* FilenamePrefix = nullptr;
diff --git a/Engine/Source/Developer/ShaderFormatOpenGL/Private/OpenGLShaderCompiler.cpp b/Engine/Source/Developer/ShaderFormatOpenGL/Private/OpenGLShaderCompiler.cpp
--- a/Engine/Source/Developer/ShaderFormatOpenGL/Private/OpenGLShaderCompiler.cpp
+++ b/Engine/Source/Developer/ShaderFormatOpenGL/Private/OpenGLShaderCompiler.cpp
@@ -2990,7 +2990,6 @@ static bool CompileToGlslWithShaderConductor(
 	const FString&				WorkingDirectory,
 	GLSLVersion					Version,
 	const EShaderFrequency		Frequency,
-	uint32						CCFlags,
 	const FString&				PreprocessedShader,
 	char*&						OutGlslShaderSource)
 {
@@ -3057,9 +3056,7 @@ static bool CompileToGlslWithShaderConductor(
 
 	PRAGMA_ENABLE_DEPRECATION_WARNINGS
 
-	UE::ShaderCompilerCommon::FDebugShaderDataOptions DebugDataOptions;
-	DebugDataOptions.HlslCCFlags = CCFlags;
-	UE::ShaderCompilerCommon::DumpDebugShaderData(Input, PreprocessedShader, DebugDataOptions);
+	UE::ShaderCompilerCommon::DumpDebugShaderData(Input, PreprocessedShader);
 
 	uint32_t BlendFlags = GetDecalBlendFlags(Input);
 
@@ -3426,14 +3423,12 @@ void FOpenGLFrontend::CompileShader(const FShaderCompilerInput& Input, FShaderCo
 #if DXC_SUPPORTED
 	if (bUseSC)
 	{
-		bCompilationSucceeded = CompileToGlslWithShaderConductor(Input, Output, WorkingDirectory, Version, Frequency, CCFlags, PreprocessedShader, GlslShaderSource);
+		bCompilationSucceeded = CompileToGlslWithShaderConductor(Input, Output, WorkingDirectory, Version, Frequency, PreprocessedShader, GlslShaderSource);
 	}
 	else
 #endif // DXC_SUPPORTED
 	{
-		UE::ShaderCompilerCommon::FDebugShaderDataOptions DebugDataOptions;
-		DebugDataOptions.HlslCCFlags = CCFlags;
-		UE::ShaderCompilerCommon::DumpDebugShaderData(Input, PreprocessedShader, DebugDataOptions);
+		UE::ShaderCompilerCommon::DumpDebugShaderData(Input, PreprocessedShader);
 		
 		CCFlags |= HLSLCC_NoValidation;
 		FGlslCodeBackend* BackEnd = CreateBackend(Version, CCFlags, HlslCompilerTarget);
diff --git a/Engine/Source/Developer/ShaderFormatVectorVM/Private/VectorVMShaderCompiler.cpp b/Engine/Source/Developer/ShaderFormatVectorVM/Private/VectorVMShaderCompiler.cpp
--- a/Engine/Source/Developer/ShaderFormatVectorVM/Private/VectorVMShaderCompiler.cpp
+++ b/Engine/Source/Developer/ShaderFormatVectorVM/Private/VectorVMShaderCompiler.cpp
@@ -136,9 +136,7 @@ bool CompileShader_VectorVM(const FShaderCompilerInput& Input, FShaderCompilerOu
 	// Required as we added the RemoveUniformBuffersFromSource() function (the cross-compiler won't be able to interpret comments w/o a preprocessor)
 	//CCFlags &= ~HLSLCC_NoPreprocess;
 
-	UE::ShaderCompilerCommon::FDebugShaderDataOptions DebugDataOptions;
-	DebugDataOptions.HlslCCFlags = CCFlags;
-	UE::ShaderCompilerCommon::DumpDebugShaderData(Input, PreprocessedShader, DebugDataOptions);
+	UE::ShaderCompilerCommon::DumpDebugShaderData(Input, PreprocessedShader);
 
 	FVectorVMCodeBackend VVMBackEnd(CCFlags, HlslCompilerTarget, VMCompilationOutput);
 		FVectorVMLanguageSpec VVMLanguageSpec; 
diff --git a/Engine/Source/Developer/VulkanShaderFormat/Private/VulkanShaderCompiler.cpp b/Engine/Source/Developer/VulkanShaderFormat/Private/VulkanShaderCompiler.cpp
--- a/Engine/Source/Developer/VulkanShaderFormat/Private/VulkanShaderCompiler.cpp
+++ b/Engine/Source/Developer/VulkanShaderFormat/Private/VulkanShaderCompiler.cpp
@@ -1433,7 +1433,6 @@ static void BuildShaderOutput(
 FCompilerInfo::FCompilerInfo(const FShaderCompilerInput& InInput, const FString& InWorkingDirectory, EHlslShaderFrequency InFrequency) :
 	Input(InInput),
 	WorkingDirectory(InWorkingDirectory),
-	CCFlags(0),
 	Frequency(InFrequency)
 {
 	BaseSourceFilename = Input.GetSourceFilename();
@@ -2397,8 +2396,6 @@ static bool CompileWithShaderConductor(
 		}
 	}
 
-	UE::ShaderCompilerCommon::DumpDebugShaderData(Input, PreprocessedShader, { CompilerInfo.CCFlags });
-
 	if (bDebugDump)
 	{
 		VulkanCreateDXCCompileBatchFiles(
@@ -2655,38 +2652,7 @@ const FString EntryPointName = Input.EntryPointName;
 
 	FCompilerInfo CompilerInfo(Input, WorkingDirectory, HlslFrequency);
 
-	// Setup hlslcc flags. Needed here as it will be used when dumping debug info
-	{
-		CompilerInfo.CCFlags |= HLSLCC_PackUniforms;
-		CompilerInfo.CCFlags |= HLSLCC_PackUniformsIntoUniformBuffers;
-
-		// Only flatten structures inside UBs
-		CompilerInfo.CCFlags |= HLSLCC_FlattenUniformBufferStructures;
-
-		if (Input.Environment.FullPrecisionInPS)
-		{
-			CompilerInfo.CCFlags |= HLSLCC_UseFullPrecisionInPS;
-		}
-
-		CompilerInfo.CCFlags |= HLSLCC_SeparateShaderObjects;
-		CompilerInfo.CCFlags |= HLSLCC_KeepSamplerAndImageNames;
-
-		CompilerInfo.CCFlags |= HLSLCC_RetainSizes;
-
-		// ES doesn't support origin layout
-		CompilerInfo.CCFlags |= HLSLCC_DX11ClipSpace;
-
-		// Required as we added the RemoveUniformBuffersFromSource() function (the cross-compiler won't be able to interpret comments w/o a preprocessor)
-		CompilerInfo.CCFlags &= ~HLSLCC_NoPreprocess;
-
-		if (!bDirectCompile || UE_BUILD_DEBUG)
-		{
-			// Validation is expensive - only do it when compiling directly for debugging
-			CompilerInfo.CCFlags |= HLSLCC_NoValidation;
-		}
-	}
-
-	UE::ShaderCompilerCommon::DumpDebugShaderData(Input, PreprocessedShaderSource, { CompilerInfo.CCFlags });
+	UE::ShaderCompilerCommon::DumpDebugShaderData(Input, PreprocessedShaderSource);
 
 	FVulkanBindingTable BindingTable(CompilerInfo.Frequency);
 	bool bSuccess = false;
diff --git a/Engine/Source/Developer/VulkanShaderFormat/Private/VulkanShaderFormat.h b/Engine/Source/Developer/VulkanShaderFormat/Private/VulkanShaderFormat.h
--- a/Engine/Source/Developer/VulkanShaderFormat/Private/VulkanShaderFormat.h
+++ b/Engine/Source/Developer/VulkanShaderFormat/Private/VulkanShaderFormat.h
@@ -30,7 +30,6 @@ struct FCompilerInfo
 	const struct FShaderCompilerInput& Input;
 	FString WorkingDirectory;
 	FString Profile;
-	uint32 CCFlags;
 	EHlslShaderFrequency Frequency;
 	FString BaseSourceFilename;
 
