[shaders] implement independent preprocessing API for vulkan shader format

diff --git a/Engine/Source/Developer/VulkanShaderFormat/Private/VulkanShaderCompiler.cpp b/Engine/Source/Developer/VulkanShaderFormat/Private/VulkanShaderCompiler.cpp
--- a/Engine/Source/Developer/VulkanShaderFormat/Private/VulkanShaderCompiler.cpp
+++ b/Engine/Source/Developer/VulkanShaderFormat/Private/VulkanShaderCompiler.cpp
@@ -2,15 +2,17 @@
 // .
 
 #include "VulkanShaderFormat.h"
-#include "VulkanCommon.h"
-#include "ShaderPreprocessor.h"
-#include "ShaderCompilerCommon.h"
-#include "ShaderParameterParser.h"
+
 #include "HlslccHeaderWriter.h"
 #include "hlslcc.h"
-#include "SpirvReflectCommon.h"
 #include "RHIShaderFormatDefinitions.inl"
 #include "Runtime/RenderCore/Internal/ShaderCompilerDefinitions.h"
+#include "ShaderCompilerCommon.h"
+#include "ShaderParameterParser.h"
+#include "ShaderPreprocessor.h"
+#include "ShaderPreprocessTypes.h"
+#include "SpirvReflectCommon.h"
+#include "VulkanCommon.h"
 
 #if PLATFORM_MAC
 // Horrible hack as we need the enum available but the Vulkan headers do not compile on Mac
@@ -67,6 +69,16 @@
 		|| ShaderFormat == NAME_VULKAN_SM5_ANDROID;
 }
 
+enum class EVulkanShaderVersion
+{
+	ES3_1,
+	ES3_1_ANDROID,
+	SM5,
+	SM5_ANDROID,
+	SM6,
+	Invalid,
+};
+
 inline CrossCompiler::FShaderConductorOptions::ETargetEnvironment GetMinimumTargetEnvironment(EVulkanShaderVersion ShaderVersion)
 {
 	return (ShaderVersion == EVulkanShaderVersion::SM6) ?
@@ -2497,52 +2509,49 @@
 #endif // PLATFORM_MAC || PLATFORM_WINDOWS || PLATFORM_LINUX
 
 
-void DoCompileVulkanShader(const FShaderCompilerInput& Input, FShaderCompilerOutput& Output, const class FString& WorkingDirectory, EVulkanShaderVersion Version)
+
+
+EVulkanShaderVersion FormatToVersion(FName Format)
 {
-	check(IsVulkanShaderFormat(Input.ShaderFormat));
+	if (Format == NAME_VULKAN_ES3_1)
+	{
+		return EVulkanShaderVersion::ES3_1;
+	}
+	else if (Format == NAME_VULKAN_ES3_1_ANDROID)
+	{
+		return EVulkanShaderVersion::ES3_1_ANDROID;
+	}
+	else if (Format == NAME_VULKAN_SM5_ANDROID)
+	{
+		return EVulkanShaderVersion::SM5_ANDROID;
+	}
+	else if (Format == NAME_VULKAN_SM5)
+	{
+		return EVulkanShaderVersion::SM5;
+	}
+	else if (Format == NAME_VULKAN_SM6)
+	{
+		return EVulkanShaderVersion::SM6;
+	}
+	else
+	{
+		FString FormatStr = Format.ToString();
+		checkf(0, TEXT("Invalid shader format passed to Vulkan shader compiler: %s"), *FormatStr);
+		return EVulkanShaderVersion::Invalid;
+	}
+}
 
+bool PreprocessVulkanShader(const FShaderCompilerInput& Input, const FShaderCompilerEnvironment& Environment, FShaderPreprocessOutput& PreprocessOutput)
+{
+	EVulkanShaderVersion Version = FormatToVersion(Input.ShaderFormat);
+
+	PRAGMA_DISABLE_DEPRECATION_WARNINGS		// FShaderCompilerDefinitions will be made internal in the future, marked deprecated until then
+
 	const bool bIsSM6 = (Version == EVulkanShaderVersion::SM6);
 	const bool bIsSM5 = (Version == EVulkanShaderVersion::SM5) || (Version == EVulkanShaderVersion::SM5_ANDROID);
 	const bool bIsMobileES31 = (Version == EVulkanShaderVersion::ES3_1 || Version == EVulkanShaderVersion::ES3_1_ANDROID);
-	bool bStripReflect = Input.IsRayTracingShader();
-	// By default we strip reflecion information for Android platform to avoid issues with older drivers
-	if (IsAndroidShaderFormat(Input.ShaderFormat))
-	{
-		bStripReflect = Input.Environment.GetCompileArgument(TEXT("STRIP_REFLECT_ANDROID"), true);
-	}
-
 	const CrossCompiler::FShaderConductorOptions::ETargetEnvironment MinTargetEnvironment = GetMinimumTargetEnvironment(Version);
-
-	const EHlslShaderFrequency FrequencyTable[] =
-	{
-		HSF_VertexShader,
-		HSF_InvalidFrequency,
-		HSF_InvalidFrequency,
-		HSF_PixelShader,
-		(bIsSM5 || bIsSM6) ? HSF_GeometryShader : HSF_InvalidFrequency,
-		HSF_ComputeShader, 
-		(bIsSM5 || bIsSM6) ? HSF_RayGen : HSF_InvalidFrequency,
-		(bIsSM5 || bIsSM6) ? HSF_RayMiss : HSF_InvalidFrequency,
-		(bIsSM5 || bIsSM6) ? HSF_RayHitGroup : HSF_InvalidFrequency,
-		(bIsSM5 || bIsSM6) ? HSF_RayCallable : HSF_InvalidFrequency,
-	};
-
-	const EShaderFrequency Frequency = (EShaderFrequency)Input.Target.Frequency;
-
-	const EHlslShaderFrequency HlslFrequency = FrequencyTable[Input.Target.Frequency];
-	if (HlslFrequency == HSF_InvalidFrequency)
-	{
-		Output.bSucceeded = false;
-		FShaderCompilerError& NewError = Output.Errors.AddDefaulted_GetRef();
-		NewError.StrippedErrorMessage = FString::Printf(
-			TEXT("%s shaders not supported for use in Vulkan."),
-			CrossCompiler::GetFrequencyName(Frequency));
-		return;
-	}
-
-	PRAGMA_DISABLE_DEPRECATION_WARNINGS		// FShaderCompilerDefinitions will be made internal in the future, marked deprecated until then
-
-	FString PreprocessedShader;
+	
 	FShaderCompilerDefinitions AdditionalDefines;
 	AdditionalDefines.SetDefine(TEXT("COMPILER_HLSLCC"), 1);
 	AdditionalDefines.SetDefine(TEXT("COMPILER_VULKAN"), 1);
@@ -2602,16 +2611,14 @@
 
 	PRAGMA_ENABLE_DEPRECATION_WARNINGS
 
-	const double StartPreprocessTime = FPlatformTime::Seconds();
-
 	// Preprocess the shader.
-	FString PreprocessedShaderSource;
+	FString& PreprocessedShaderSource = PreprocessOutput.EditSource();
 	const bool bDirectCompile = FParse::Param(FCommandLine::Get(), TEXT("directcompile"));
 	if (bDirectCompile)
 	{
 		if (!FFileHelper::LoadFileToString(PreprocessedShaderSource, *Input.VirtualSourceFilePath))
 		{
-			return;
+			return false;
 		}
 
 		// Remove const as we are on debug-only mode
@@ -2619,22 +2626,19 @@
 	}
 	else
 	{
-		if (!PreprocessShader(PreprocessedShaderSource, Output, Input, AdditionalDefines))
+		if (!PreprocessShader(PreprocessOutput, Input, Environment, AdditionalDefines))
 		{
 			// The preprocessing stage will add any relevant errors.
-			return;
+			return false;
 		}
 	}
 
-	FShaderParameterParser ShaderParameterParser(Input.Environment.CompilerFlags, nullptr);
-	if (!ShaderParameterParser.ParseAndModify(Input, Output.Errors, PreprocessedShaderSource, EBindlessParameterMode::Vulkan))
+	if (!PreprocessOutput.ParseAndModify(Input, Environment, EBindlessParameterMode::Vulkan))
 	{
 		// The FShaderParameterParser will add any relevant errors.
-		return;
+		return false;
 	}
 
-const FString EntryPointName = Input.EntryPointName;
-
 	RemoveUniformBuffersFromSource(Input.Environment, PreprocessedShaderSource);
 
 	// Process TEXT macro.
@@ -2644,25 +2648,69 @@
 	#if UE_VULKAN_SHADER_COMPILER_ALLOW_DEAD_CODE_REMOVAL
 	if (Input.Environment.CompilerFlags.Contains(CFLAG_RemoveDeadCode))
 	{
-		UE::ShaderCompilerCommon::RemoveDeadCode(PreprocessedShaderSource, EntryPointName, Output.Errors);
+		UE::ShaderCompilerCommon::RemoveDeadCode(PreprocessedShaderSource, Input.EntryPointName, PreprocessOutput.EditErrors());
 	}
 	#endif // UE_VULKAN_SHADER_COMPILER_ALLOW_DEAD_CODE_REMOVAL
 
-	Output.PreprocessTime = FPlatformTime::Seconds() - StartPreprocessTime;
+	return true;
+}
 
+void CompileVulkanShader(const FShaderCompilerInput& Input, const FShaderPreprocessOutput& PreprocessOutput, FShaderCompilerOutput& Output, const class FString& WorkingDirectory)
+{
+	check(IsVulkanShaderFormat(Input.ShaderFormat));
+	EVulkanShaderVersion Version = FormatToVersion(Input.ShaderFormat);
+
+	const bool bIsSM6 = (Version == EVulkanShaderVersion::SM6);
+	const bool bIsSM5 = (Version == EVulkanShaderVersion::SM5) || (Version == EVulkanShaderVersion::SM5_ANDROID);
+	const bool bIsMobileES31 = (Version == EVulkanShaderVersion::ES3_1 || Version == EVulkanShaderVersion::ES3_1_ANDROID);
+	bool bStripReflect = Input.IsRayTracingShader();
+	// By default we strip reflecion information for Android platform to avoid issues with older drivers
+	if (IsAndroidShaderFormat(Input.ShaderFormat))
+	{
+		bStripReflect = Input.Environment.GetCompileArgument(TEXT("STRIP_REFLECT_ANDROID"), true);
+	}
+
+	const CrossCompiler::FShaderConductorOptions::ETargetEnvironment MinTargetEnvironment = GetMinimumTargetEnvironment(Version);
+
+	const EHlslShaderFrequency FrequencyTable[] =
+	{
+		HSF_VertexShader,
+		HSF_InvalidFrequency,
+		HSF_InvalidFrequency,
+		HSF_PixelShader,
+		(bIsSM5 || bIsSM6) ? HSF_GeometryShader : HSF_InvalidFrequency,
+		HSF_ComputeShader, 
+		(bIsSM5 || bIsSM6) ? HSF_RayGen : HSF_InvalidFrequency,
+		(bIsSM5 || bIsSM6) ? HSF_RayMiss : HSF_InvalidFrequency,
+		(bIsSM5 || bIsSM6) ? HSF_RayHitGroup : HSF_InvalidFrequency,
+		(bIsSM5 || bIsSM6) ? HSF_RayCallable : HSF_InvalidFrequency,
+	};
+
+	const EShaderFrequency Frequency = (EShaderFrequency)Input.Target.Frequency;
+
+	const EHlslShaderFrequency HlslFrequency = FrequencyTable[Input.Target.Frequency];
+	if (HlslFrequency == HSF_InvalidFrequency)
+	{
+		Output.bSucceeded = false;
+		FShaderCompilerError& NewError = Output.Errors.AddDefaulted_GetRef();
+		NewError.StrippedErrorMessage = FString::Printf(
+			TEXT("%s shaders not supported for use in Vulkan."),
+			CrossCompiler::GetFrequencyName(Frequency));
+		return;
+	}
+
+	const bool bDirectCompile = FParse::Param(FCommandLine::Get(), TEXT("directcompile"));
 	FCompilerInfo CompilerInfo(Input, WorkingDirectory, HlslFrequency);
 
-	UE::ShaderCompilerCommon::DumpDebugShaderData(Input, PreprocessedShaderSource);
-
 	FVulkanBindingTable BindingTable(CompilerInfo.Frequency);
 	bool bSuccess = false;
 
 #if PLATFORM_MAC || PLATFORM_WINDOWS || PLATFORM_LINUX
 	// Cross-compile shader via ShaderConductor (DXC, SPIRV-Tools, SPIRV-Cross)
-	bSuccess = CompileWithShaderConductor(PreprocessedShaderSource, EntryPointName, Frequency, CompilerInfo, Output, BindingTable, bStripReflect, MinTargetEnvironment);
+	bSuccess = CompileWithShaderConductor(PreprocessOutput.GetSource(), Input.EntryPointName, Frequency, CompilerInfo, Output, BindingTable, bStripReflect, MinTargetEnvironment);
 #endif // PLATFORM_MAC || PLATFORM_WINDOWS || PLATFORM_LINUX
 	
-	ShaderParameterParser.ValidateShaderParameterTypes(Input, bIsMobileES31, Output);
+	PreprocessOutput.GetParameterParser().ValidateShaderParameterTypes(Input, bIsMobileES31, Output);
 	
 	if (bDirectCompile)
 	{
@@ -2673,3 +2721,8 @@
 		ensure(bSuccess);
 	}
 }
+
+void OutputVulkanDebugData(const FShaderCompilerInput& Input, const FShaderPreprocessOutput& PreprocessOutput, const FShaderCompilerOutput& Output)
+{
+	UE::ShaderCompilerCommon::DumpExtendedDebugShaderData(Input, PreprocessOutput, Output);
+}
diff --git a/Engine/Source/Developer/VulkanShaderFormat/Private/VulkanShaderFormat.cpp b/Engine/Source/Developer/VulkanShaderFormat/Private/VulkanShaderFormat.cpp
--- a/Engine/Source/Developer/VulkanShaderFormat/Private/VulkanShaderFormat.cpp
+++ b/Engine/Source/Developer/VulkanShaderFormat/Private/VulkanShaderFormat.cpp
@@ -13,6 +13,22 @@
 #include "ShaderConductorContext.h"
 #include "RHIShaderFormatDefinitions.inl"
 
+extern bool PreprocessVulkanShader(
+	const struct  FShaderCompilerInput& Input,
+	const struct FShaderCompilerEnvironment& Environment,
+	class FShaderPreprocessOutput& PreprocessOutput);
+
+extern void CompileVulkanShader(
+	const struct FShaderCompilerInput& Input,
+	const class FShaderPreprocessOutput& PreprocessOutput,
+	struct FShaderCompilerOutput& Output,
+	const class FString& WorkingDirectory);
+
+extern void OutputVulkanDebugData(
+	const FShaderCompilerInput& Input, 
+	const FShaderPreprocessOutput& PreprocessOutput, 
+	const FShaderCompilerOutput& Output);
+
 class FShaderFormatVulkan : public IShaderFormat
 {
 	enum 
@@ -37,7 +53,8 @@ class FShaderFormatVulkan : public IShaderFormat
 			return UE_SHADER_VULKAN_ES3_1_VER;
 		}
 
-		check(0);
+		FString FormatStr = Format.ToString();
+		checkf(0, TEXT("Invalid shader format passed to Vulkan shader compiler: %s"), *FormatStr);
 		return -1;
 	}
 
@@ -75,39 +92,32 @@ public:
 		OutFormats.Add(NAME_VULKAN_SM6);
 	}
 
-	virtual void CompileShader(FName Format, const struct FShaderCompilerInput& Input, struct FShaderCompilerOutput& Output,const FString& WorkingDirectory) const
+	virtual bool SupportsIndependentPreprocessing() const override
 	{
-		check(InternalGetVersion(Format) >= 0);
-		if (Format == NAME_VULKAN_ES3_1)
-		{
-			DoCompileVulkanShader(Input, Output, WorkingDirectory, EVulkanShaderVersion::ES3_1);
-		}
-		else if (Format == NAME_VULKAN_ES3_1_ANDROID)
-		{
-			DoCompileVulkanShader(Input, Output, WorkingDirectory, EVulkanShaderVersion::ES3_1_ANDROID);
-		}
-		else if (Format == NAME_VULKAN_SM5_ANDROID)
-		{
-			DoCompileVulkanShader(Input, Output, WorkingDirectory, EVulkanShaderVersion::SM5_ANDROID);
-		}
-		else if (Format == NAME_VULKAN_SM5)
-		{
-			DoCompileVulkanShader(Input, Output, WorkingDirectory, EVulkanShaderVersion::SM5);
-		}
-		else if (Format == NAME_VULKAN_SM6)
-		{
-			DoCompileVulkanShader(Input, Output, WorkingDirectory, EVulkanShaderVersion::SM6);
-		}
+		return true;
 	}
 
-	virtual const TCHAR* GetPlatformIncludeDirectory() const
+	virtual bool PreprocessShader(const FShaderCompilerInput& Input, const FShaderCompilerEnvironment& Environment, FShaderPreprocessOutput& PreprocessOutput) const
 	{
-		return TEXT("Vulkan");
+		check(InternalGetVersion(Input.ShaderFormat) >= 0);
+
+		return PreprocessVulkanShader(Input, Environment, PreprocessOutput);
+	}
+
+	virtual void CompilePreprocessedShader(const FShaderCompilerInput& Input, const FShaderPreprocessOutput& PreprocessOutput, FShaderCompilerOutput& Output,const FString& WorkingDirectory) const override
+	{
+		check(InternalGetVersion(Input.ShaderFormat) >= 0);
+		CompileVulkanShader(Input, PreprocessOutput, Output, WorkingDirectory);
 	}
 
-	virtual bool UsesHLSLcc(const struct FShaderCompilerInput& Input) const override
+	virtual void OutputDebugData(const FShaderCompilerInput& Input, const FShaderPreprocessOutput& PreprocessOutput, const FShaderCompilerOutput& Output) const override
 	{
-		return !Input.Environment.CompilerFlags.Contains(CFLAG_ForceDXC);
+		OutputVulkanDebugData(Input, PreprocessOutput, Output);
+	}
+
+	virtual const TCHAR* GetPlatformIncludeDirectory() const
+	{
+		return TEXT("Vulkan");
 	}
 };
 
diff --git a/Engine/Source/Developer/VulkanShaderFormat/Private/VulkanShaderFormat.h b/Engine/Source/Developer/VulkanShaderFormat/Private/VulkanShaderFormat.h
--- a/Engine/Source/Developer/VulkanShaderFormat/Private/VulkanShaderFormat.h
+++ b/Engine/Source/Developer/VulkanShaderFormat/Private/VulkanShaderFormat.h
@@ -12,18 +12,6 @@
 #define UE_VULKAN_SHADER_COMPILER_ALLOW_DEAD_CODE_REMOVAL 1
 #endif // UE_VULKAN_SHADER_COMPILER_ALLOW_DEAD_CODE_REMOVAL
 
-enum class EVulkanShaderVersion
-{
-	ES3_1,
-	ES3_1_ANDROID,
-	SM5,
-	SM5_ANDROID,
-	SM6,
-};
-
-extern void DoCompileVulkanShader(const struct FShaderCompilerInput& Input,struct FShaderCompilerOutput& Output,const class FString& WorkingDirectory, EVulkanShaderVersion Version);
-
-
 // Hold information to be able to call the compilers
 struct FCompilerInfo
 {
diff --git a/Engine/Source/Runtime/RenderCore/Public/ShaderPreprocessTypes.h b/Engine/Source/Runtime/RenderCore/Public/ShaderPreprocessTypes.h
--- a/Engine/Source/Runtime/RenderCore/Public/ShaderPreprocessTypes.h
+++ b/Engine/Source/Runtime/RenderCore/Public/ShaderPreprocessTypes.h
@@ -135,11 +135,25 @@ public:
 		return Errors;
 	}
 
-	inline bool ParseAndModify(const FShaderCompilerInput& Input, const FShaderCompilerEnvironment& Environment, const TCHAR* ConstantBufferType, const TArrayView<const TCHAR* const> ExtraSRVTypes = {}, const TArrayView<const TCHAR* const> ExtraUAVTypes = {})
+	inline bool ParseAndModify(
+		const FShaderCompilerInput& Input, 
+		const FShaderCompilerEnvironment& Environment, 
+		const TCHAR* ConstantBufferType, 
+		const TArrayView<const TCHAR* const> ExtraSRVTypes = {}, 
+		const TArrayView<const TCHAR* const> ExtraUAVTypes = {}, 
+		EBindlessParameterMode BindlessParameterMode = EBindlessParameterMode::Default)
 	{
 		FShaderParameterParser Tmp(Environment.CompilerFlags, ConstantBufferType, ExtraSRVTypes, ExtraUAVTypes);
 		ParameterParser = MoveTemp(Tmp);
-		return ParameterParser.ParseAndModify(Input, Errors, PreprocessedSource);
+		return ParameterParser.ParseAndModify(Input, Errors, PreprocessedSource, BindlessParameterMode);
+	}
+
+	inline bool ParseAndModify(
+		const FShaderCompilerInput& Input,
+		const FShaderCompilerEnvironment& Environment,
+		EBindlessParameterMode BindlessParameterMode)
+	{
+		return ParseAndModify(Input, Environment, nullptr, {}, {}, BindlessParameterMode);
 	}
 
 	inline const FShaderParameterParser& GetParameterParser() const
