VulkanRHI: Start a clean up of the VulkanShaderCompiler, small steps at first (WIP).  Reduce parameters in shader compilation functions by batching them together in a struct (more params coming with RT). Remove FCompilerInfo class.  Replace repeated loops with lambdas.

#rnx

diff --git a/Engine/Source/Developer/VulkanShaderFormat/Private/VulkanShaderCompiler.cpp b/Engine/Source/Developer/VulkanShaderFormat/Private/VulkanShaderCompiler.cpp
--- a/Engine/Source/Developer/VulkanShaderFormat/Private/VulkanShaderCompiler.cpp
+++ b/Engine/Source/Developer/VulkanShaderFormat/Private/VulkanShaderCompiler.cpp
@@ -79,11 +79,51 @@ enum class EVulkanShaderVersion
 	Invalid,
 };
 
-inline CrossCompiler::FShaderConductorOptions::ETargetEnvironment GetMinimumTargetEnvironment(EVulkanShaderVersion ShaderVersion)
+inline EVulkanShaderVersion FormatToVersion(FName Format)
 {
-	return (ShaderVersion == EVulkanShaderVersion::SM6) ?
-		CrossCompiler::FShaderConductorOptions::ETargetEnvironment::Vulkan_1_3:
-		CrossCompiler::FShaderConductorOptions::ETargetEnvironment::Vulkan_1_1;
+	if (Format == NAME_VULKAN_ES3_1)
+	{
+		return EVulkanShaderVersion::ES3_1;
+	}
+	else if (Format == NAME_VULKAN_ES3_1_ANDROID)
+	{
+		return EVulkanShaderVersion::ES3_1_ANDROID;
+	}
+	else if (Format == NAME_VULKAN_SM5_ANDROID)
+	{
+		return EVulkanShaderVersion::SM5_ANDROID;
+	}
+	else if (Format == NAME_VULKAN_SM5)
+	{
+		return EVulkanShaderVersion::SM5;
+	}
+	else if (Format == NAME_VULKAN_SM6)
+	{
+		return EVulkanShaderVersion::SM6;
+	}
+	else
+	{
+		FString FormatStr = Format.ToString();
+		checkf(0, TEXT("Invalid shader format passed to Vulkan shader compiler: %s"), *FormatStr);
+		return EVulkanShaderVersion::Invalid;
+	}
+}
+
+inline CrossCompiler::FShaderConductorOptions::ETargetEnvironment GetMinimumTargetEnvironment(const FShaderCompilerInput& Input)
+{
+	const EVulkanShaderVersion ShaderVersion = FormatToVersion(Input.ShaderFormat);
+	if (ShaderVersion == EVulkanShaderVersion::SM6)
+	{
+		return CrossCompiler::FShaderConductorOptions::ETargetEnvironment::Vulkan_1_3;
+	}
+	else if (Input.IsRayTracingShader() || Input.Environment.CompilerFlags.Contains(CFLAG_InlineRayTracing))
+	{
+		return CrossCompiler::FShaderConductorOptions::ETargetEnvironment::Vulkan_1_2;
+	}
+	else
+	{
+		return CrossCompiler::FShaderConductorOptions::ETargetEnvironment::Vulkan_1_1;
+	}
 }
 
 DEFINE_LOG_CATEGORY_STATIC(LogVulkanShaderCompiler, Log, All); 
@@ -275,6 +315,86 @@ struct FOLDVulkanCodeHeader
 	TArray<uint32> NEWEmulatedUBCopyRanges;
 };
 
+
+// A collection of states and data that is locked in at the top level call and doesn't change throughout the compilation process
+struct VulkanShaderCompilerInternalState
+{
+	VulkanShaderCompilerInternalState(const FShaderCompilerInput& InInput)
+		: Input(InInput)
+		, Version(FormatToVersion(Input.ShaderFormat))
+		, MinimumTargetEnvironment(GetMinimumTargetEnvironment(InInput))
+		, bStripReflect(InInput.IsRayTracingShader() || (IsAndroidShaderFormat(Input.ShaderFormat) && InInput.Environment.GetCompileArgument(TEXT("STRIP_REFLECT_ANDROID"), true)))
+		, bUseBindlessUniformBuffer(InInput.IsRayTracingShader() && ((EShaderFrequency)InInput.Target.Frequency != SF_RayGen))
+		, bIsRayHitGroupShader(InInput.IsRayTracingShader() && ((EShaderFrequency)InInput.Target.Frequency == SF_RayHitGroup))
+		, bSupportsBindless(InInput.Environment.CompilerFlags.Contains(CFLAG_BindlessResources) || InInput.Environment.CompilerFlags.Contains(CFLAG_BindlessSamplers))
+		, bUseRootParameterStructure(InInput.IsRayTracingShader() && InInput.RootParametersStructure)
+		, bDebugDump(InInput.DumpDebugInfoEnabled())
+	{
+	}
+
+	const FShaderCompilerInput& Input;
+
+	const EVulkanShaderVersion Version;
+	const CrossCompiler::FShaderConductorOptions::ETargetEnvironment MinimumTargetEnvironment;
+
+	const bool bStripReflect;
+	const bool bUseBindlessUniformBuffer;
+	const bool bIsRayHitGroupShader;
+
+	const bool bSupportsBindless;
+	const bool bUseRootParameterStructure;
+	const bool bDebugDump;
+
+	// Forwarded calls for convenience
+	inline EShaderFrequency GetShaderFrequency() const
+	{
+		return static_cast<EShaderFrequency>(Input.Target.Frequency);
+	}
+	inline const FString& GetEntryPointName() const
+	{
+		return Input.EntryPointName;
+	}
+	inline bool IsRayTracingShader() const
+	{
+		return Input.IsRayTracingShader();
+	}
+	inline bool IsSM6() const
+	{
+		return (Version == EVulkanShaderVersion::SM6);
+	}
+	inline bool IsSM5() const
+	{
+		return (Version == EVulkanShaderVersion::SM5) || (Version == EVulkanShaderVersion::SM5_ANDROID);
+	}
+	inline bool IsMobileES31() const
+	{
+		return (Version == EVulkanShaderVersion::ES3_1 || Version == EVulkanShaderVersion::ES3_1_ANDROID);
+	}
+	inline EHlslShaderFrequency GetHlslShaderFrequency() const
+	{
+		const EHlslShaderFrequency FrequencyTable[] =
+		{
+			HSF_VertexShader,
+			HSF_InvalidFrequency,
+			HSF_InvalidFrequency,
+			HSF_PixelShader,
+			(IsSM5() || IsSM6()) ? HSF_GeometryShader : HSF_InvalidFrequency,
+			HSF_ComputeShader,
+			(IsSM5() || IsSM6()) ? HSF_RayGen : HSF_InvalidFrequency,
+			(IsSM5() || IsSM6()) ? HSF_RayMiss : HSF_InvalidFrequency,
+			(IsSM5() || IsSM6()) ? HSF_RayHitGroup : HSF_InvalidFrequency,
+			(IsSM5() || IsSM6()) ? HSF_RayCallable : HSF_InvalidFrequency,
+		};
+		return FrequencyTable[Input.Target.Frequency];
+	}
+	inline FString GetDebugName() const
+	{
+		return Input.DumpDebugInfoPath.Right(Input.DumpDebugInfoPath.Len() - Input.DumpDebugInfoRootPath.Len());
+	}
+};
+
+
+
 static void AddImmutable(FVulkanShaderHeader& OutHeader, int32 GlobalIndex)
 {
 	check(GlobalIndex < UINT16_MAX);
@@ -1442,13 +1562,6 @@ static void BuildShaderOutput(
 	CullGlobalUniformBuffers(ShaderInput.Environment.UniformBufferMap, ShaderOutput.ParameterMap);
 }
 
-FCompilerInfo::FCompilerInfo(const FShaderCompilerInput& InInput, const FString& InWorkingDirectory, EHlslShaderFrequency InFrequency) :
-	Input(InInput),
-	WorkingDirectory(InWorkingDirectory),
-	Frequency(InFrequency)
-{
-	BaseSourceFilename = Input.GetSourceFilename();
-}
 
 #if PLATFORM_MAC || PLATFORM_WINDOWS || PLATFORM_LINUX
 
@@ -1588,26 +1701,16 @@ static uint32 CalculateSpirvInstructionCount(FVulkanSpirv& Spirv)
 static bool BuildShaderOutputFromSpirv(
 	CrossCompiler::FShaderConductorContext&	CompilerContext,
 	FVulkanSpirv&							Spirv,
-	const FShaderCompilerInput&				Input,
+	const VulkanShaderCompilerInternalState InternalState,
 	FShaderCompilerOutput&					Output,
 	FVulkanBindingTable&					BindingTable,
-	const FString&							EntryPointName,
-	uint8									WaveSize,
-	bool									bStripReflect,
-	bool									bIsRayTracingShader,
-	bool									bDebugDump)
+	uint8									WaveSize
+)
 {
+	const bool bIsRayTracingShader = InternalState.IsRayTracingShader();
+	const EShaderFrequency ShaderFrequency = InternalState.GetShaderFrequency();
 	FShaderParameterMap& ParameterMap = Output.ParameterMap;
 
-	SpvReflectResult SpvResult = SPV_REFLECT_RESULT_SUCCESS;
-
-	uint8 UAVIndices = 0xff;
-	uint64 TextureIndices = 0xffffffffffffffff;
-	uint16 SamplerIndices = 0xffff;
-	const uint32 GlobalSetId = 32;
-
-	TMap<const SpvReflectDescriptorBinding*, uint32> SamplerStatesUseCount;
-
 	// Reflect SPIR-V module with SPIRV-Reflect library
 	const size_t SpirvDataSize = Spirv.GetByteSize();
 	spv_reflect::ShaderModule Reflection(SpirvDataSize, Spirv.GetByteData(), SPV_REFLECT_RETURN_FLAG_SAMPLER_IMAGE_USAGE);
@@ -1620,7 +1723,7 @@ static bool BuildShaderOutputFromSpirv(
 	{
 		// For now only use the primary entry point for hit groups until we decide how to best support hit group library shaders.
 		FString OutMain, OutAnyHit, OutIntersection;
-		UE::ShaderCompilerCommon::ParseRayTracingEntryPoint(EntryPointName, OutMain, OutAnyHit, OutIntersection);
+		UE::ShaderCompilerCommon::ParseRayTracingEntryPoint(InternalState.GetEntryPointName(), OutMain, OutAnyHit, OutIntersection);
 
 		for (uint32 i = 0; i < Reflection.GetEntryPointCount(); ++i)
 		{
@@ -1630,7 +1733,7 @@ static bool BuildShaderOutputFromSpirv(
 				break;
 			}
 		}
-		checkf(EntryPointIndex >= 0, TEXT("Failed to find entry point %s in SPIRV-V module."), *EntryPointName);
+		checkf(EntryPointIndex >= 0, TEXT("Failed to find entry point %s in SPIRV-V module."), *InternalState.GetEntryPointName());
 	}
 
 	// Change final entry point name in SPIR-V module
@@ -1640,24 +1743,8 @@ static bool BuildShaderOutputFromSpirv(
 		check(Result == SPV_REFLECT_RESULT_SUCCESS);
 	}
 
-	const bool bSupportsBindless = Input.Environment.CompilerFlags.Contains(CFLAG_BindlessResources) || Input.Environment.CompilerFlags.Contains(CFLAG_BindlessSamplers);
-	const EShaderFrequency Frequency = static_cast<EShaderFrequency>(Input.Target.Frequency);
-
 	FSpirvReflectBindings Bindings;
-	GatherSpirvReflectionBindings(Reflection, Bindings, Frequency, bSupportsBindless);
-
-	// Register how often a sampler-state is used
-	for (const SpvReflectDescriptorBinding* Binding : Bindings.TextureSRVs)
-	{
-		if (Binding->usage_binding_count > 0)
-		{
-			for (uint32 UsageIndex = 0; UsageIndex < Binding->usage_binding_count; ++UsageIndex)
-			{
-				const SpvReflectDescriptorBinding* AssociatedResource = Binding->usage_bindings[UsageIndex];
-				SamplerStatesUseCount.FindOrAdd(AssociatedResource)++;
-			}
-		}
-	}
+	GatherSpirvReflectionBindings(Reflection, Bindings, ShaderFrequency, InternalState.bSupportsBindless);
 
 	// Build binding table
 	TMap<const SpvReflectDescriptorBinding*, int32> BindingToIndexMap;
@@ -1698,53 +1785,25 @@ static bool BuildShaderOutputFromSpirv(
 		BindingTable.InputAttachmentsMask |= (1u << Binding->input_attachment_index);
 	}
 
-	for (const SpvReflectDescriptorBinding* Binding : Bindings.TBufferUAVs)
-	{
-		int32 BindingIndex = BindingTable.RegisterBinding(Binding->name, "u", EVulkanBindingType::StorageTexelBuffer);
-		BindingToIndexMap.Add(Binding, BindingIndex);
-	}
-
-	for (const SpvReflectDescriptorBinding* Binding : Bindings.SBufferUAVs)
+	auto RegisterBindings = [&BindingTable, &BindingToIndexMap](TArray<SpvReflectDescriptorBinding*>& Bindings, const char* BlockName, EVulkanBindingType::EType BindingType)
 	{
-		int32 BindingIndex = BindingTable.RegisterBinding(Binding->name, "u", EVulkanBindingType::StorageBuffer);
-		BindingToIndexMap.Add(Binding, BindingIndex);
-	}
-
-	for (const SpvReflectDescriptorBinding* Binding : Bindings.TextureUAVs)
-	{
-		int32 BindingIndex = BindingTable.RegisterBinding(Binding->name, "u", EVulkanBindingType::StorageImage);
-		BindingToIndexMap.Add(Binding, BindingIndex);
-	}
-
-	for (const SpvReflectDescriptorBinding* Binding : Bindings.TBufferSRVs)
-	{
-		int32 BindingIndex = BindingTable.RegisterBinding(Binding->name, "s", EVulkanBindingType::UniformTexelBuffer);
-		BindingToIndexMap.Add(Binding, BindingIndex);
-	}
-
-	for (const SpvReflectDescriptorBinding* Binding : Bindings.SBufferSRVs)
-	{
-		int32 BindingIndex = BindingTable.RegisterBinding(Binding->name, "s", EVulkanBindingType::UniformTexelBuffer);
-		BindingToIndexMap.Add(Binding, BindingIndex);
-	}
+		for (const SpvReflectDescriptorBinding* Binding : Bindings)
+		{
+			const int32 BindingIndex = BindingTable.RegisterBinding(Binding->name, BlockName, BindingType);
+			BindingToIndexMap.Add(Binding, BindingIndex);
+		}
+	};
 
-	for (const SpvReflectDescriptorBinding* Binding : Bindings.TextureSRVs)
-	{
-		int32 BindingIndex = BindingTable.RegisterBinding(Binding->name, "s", EVulkanBindingType::Image);
-		BindingToIndexMap.Add(Binding, BindingIndex);
-	}
+	RegisterBindings(Bindings.TBufferUAVs, "u", EVulkanBindingType::StorageTexelBuffer);
+	RegisterBindings(Bindings.SBufferUAVs, "u", EVulkanBindingType::StorageBuffer);
+	RegisterBindings(Bindings.TextureUAVs, "u", EVulkanBindingType::StorageImage);
 
-	for (const SpvReflectDescriptorBinding* Binding : Bindings.Samplers)
-	{
-		int32 BindingIndex = BindingTable.RegisterBinding(Binding->name, "z", EVulkanBindingType::Sampler);
-		BindingToIndexMap.Add(Binding, BindingIndex);
-	}
+	RegisterBindings(Bindings.TBufferSRVs, "s", EVulkanBindingType::UniformTexelBuffer);
+	RegisterBindings(Bindings.SBufferSRVs, "s", EVulkanBindingType::UniformTexelBuffer);
+	RegisterBindings(Bindings.TextureSRVs, "s", EVulkanBindingType::Image);
 
-	for (const SpvReflectDescriptorBinding* Binding : Bindings.AccelerationStructures)
-	{
-		int32 BindingIndex = BindingTable.RegisterBinding(Binding->name, "r", EVulkanBindingType::AccelerationStructure);
-		BindingToIndexMap.Add(Binding, BindingIndex);
-	}
+	RegisterBindings(Bindings.Samplers, "z", EVulkanBindingType::Sampler);
+	RegisterBindings(Bindings.AccelerationStructures, "r", EVulkanBindingType::AccelerationStructure);
 
 	// Sort binding table
 	BindingTable.SortBindings();
@@ -1762,219 +1821,155 @@ static bool BuildShaderOutputFromSpirv(
 		CCHeaderWriter.WriteOutputAttribute(*Attribute);
 	}
 
-	int32 UBOBindings = 0, UAVBindings = 0, SRVBindings = 0, SMPBindings = 0, GLOBindings = 0, ASBindings = 0;
-
-	auto MapDescriptorBindingToIndex = [&BindingTable, &BindingToIndexMap](const SpvReflectDescriptorBinding* Binding)
-	{
-		return BindingTable.GetRealBindingIndex(BindingToIndexMap[Binding]);
-	};
-
-	for (const SpvReflectDescriptorBinding* Binding : Bindings.UniformBuffers)
+	// Final descriptor binding numbers for UBs
 	{
-		const FString ResourceName(ANSI_TO_TCHAR(Binding->name));
+		const int32 StageOffset = InternalState.bSupportsBindless ? (ShaderStage::GetStageForFrequency(ShaderFrequency) * VulkanBindless::MaxUniformBuffersPerStage) : 0;
+		const uint32_t DescSetNumber = InternalState.bSupportsBindless ? (uint32_t)VulkanBindless::BindlessSingleUseUniformBufferSet : (uint32_t)SPV_REFLECT_SET_NUMBER_DONT_CHANGE;
 
-		if (ResourceName == UBOGlobalsNameSpv)
+		int32 UBOBindings = 0;
+		auto AddUBReflectionInfo = [&](const SpvReflectDescriptorBinding* Binding, const FString& BindingName, const TCHAR* ResourceName)
 		{
 			// Register binding for uniform buffer
-			const int32 StageOffset = bSupportsBindless ? (ShaderStage::GetStageForFrequency(Frequency) * VulkanBindless::MaxUniformBuffersPerStage) : 0;
-			const int32 BindingIndex = MapDescriptorBindingToIndex(Binding) + StageOffset;
-			const uint32_t DescSetNumber = bSupportsBindless ? (uint32_t)VulkanBindless::BindlessSingleUseUniformBufferSet : (uint32_t)SPV_REFLECT_SET_NUMBER_DONT_CHANGE;
+			const int32 BindingIndex = BindingTable.GetRealBindingIndex(BindingToIndexMap[Binding]) + StageOffset;
 
-			SpvResult = Reflection.ChangeDescriptorBindingNumbers(Binding, BindingIndex, DescSetNumber);
+			SpvReflectResult SpvResult = Reflection.ChangeDescriptorBindingNumbers(Binding, BindingIndex, DescSetNumber);
 			check(SpvResult == SPV_REFLECT_RESULT_SUCCESS);
 
-			Spirv.ReflectionInfo.Add(FVulkanSpirv::FEntry(UBOGlobalsNameSpv, BindingIndex));
+			Spirv.ReflectionInfo.Add(FVulkanSpirv::FEntry(BindingName, BindingIndex));
 
-			CCHeaderWriter.WriteUniformBlock(TEXT("_Globals_h"), UBOBindings++);
+			CCHeaderWriter.WriteUniformBlock(ResourceName, UBOBindings++);
+		};
 
-			// Register all uniform buffer members as loose data
-			FString MbrString;
+		// Process Globals first
+		for (const SpvReflectDescriptorBinding* Binding : Bindings.UniformBuffers)
+		{
+			const FString ResourceName(ANSI_TO_TCHAR(Binding->name));
 
-			for (uint32 MemberIndex = 0; MemberIndex < Binding->block.member_count; ++MemberIndex)
+			if (ResourceName == UBOGlobalsNameSpv)
 			{
-				const SpvReflectBlockVariable* Member = &(Binding->block.members[MemberIndex]);
-				CCHeaderWriter.WritePackedGlobal(ANSI_TO_TCHAR(Member->name), CrossCompiler::EPackedTypeName::HighP, Member->absolute_offset, Member->size);
-			}
+				AddUBReflectionInfo(Binding, UBOGlobalsNameSpv, TEXT("_Globals_h"));
 
-			// Stop after we found $Globals uniform buffer
-			break;
-		}
-	}
+				// Register all uniform buffer members as loose data
+				for (uint32 MemberIndex = 0; MemberIndex < Binding->block.member_count; ++MemberIndex)
+				{
+					const SpvReflectBlockVariable* Member = &(Binding->block.members[MemberIndex]);
+					CCHeaderWriter.WritePackedGlobal(ANSI_TO_TCHAR(Member->name), CrossCompiler::EPackedTypeName::HighP, Member->absolute_offset, Member->size);
+				}
 
-	for (const SpvReflectDescriptorBinding* Binding : Bindings.UniformBuffers)
-	{
-		const FString ResourceName(ANSI_TO_TCHAR(Binding->name));
+				// Stop after we found $Globals uniform buffer
+				break;
+			}
+		}
 
-		if (ResourceName != UBOGlobalsNameSpv)
+		// Process all other UB
+		for (const SpvReflectDescriptorBinding* Binding : Bindings.UniformBuffers)
 		{
-			// Register uniform buffer
-			const int32 StageOffset = bSupportsBindless ? (ShaderStage::GetStageForFrequency(Frequency) * VulkanBindless::MaxUniformBuffersPerStage) : 0;
-			const int32 BindingIndex = MapDescriptorBindingToIndex(Binding) + StageOffset;
-			const uint32_t DescSetNumber = bSupportsBindless ? (uint32_t)VulkanBindless::BindlessSingleUseUniformBufferSet : (uint32_t)SPV_REFLECT_SET_NUMBER_DONT_CHANGE;
-			SpvResult = Reflection.ChangeDescriptorBindingNumbers(Binding, BindingIndex, DescSetNumber);
-			check(SpvResult == SPV_REFLECT_RESULT_SUCCESS);
-
-			Spirv.ReflectionInfo.Add(FVulkanSpirv::FEntry(ResourceName, BindingIndex));
+			const FString ResourceName(ANSI_TO_TCHAR(Binding->name));
 
-			CCHeaderWriter.WriteUniformBlock(*ResourceName, UBOBindings++);
+			if (ResourceName != UBOGlobalsNameSpv)
+			{
+				AddUBReflectionInfo(Binding, ResourceName, *ResourceName);
+			}
 		}
 	}
 
-	for (const SpvReflectDescriptorBinding* Binding : Bindings.InputAttachments)
-	{
-		int32 BindingIndex = MapDescriptorBindingToIndex(Binding);
-
-		SpvResult = Reflection.ChangeDescriptorBindingNumbers(Binding, BindingIndex);//, GlobalSetId);
-		check(SpvResult == SPV_REFLECT_RESULT_SUCCESS);
-
-		const FString ResourceName(ANSI_TO_TCHAR(Binding->name));
-
-		Spirv.ReflectionInfo.Add(FVulkanSpirv::FEntry(ResourceName, BindingIndex));
-	}
-
-	for (const SpvReflectDescriptorBinding* Binding : Bindings.TBufferUAVs)
-	{
-		int32 BindingIndex = MapDescriptorBindingToIndex(Binding);
-
-		SpvResult = Reflection.ChangeDescriptorBindingNumbers(Binding, BindingIndex);//, GlobalSetId);
-		check(SpvResult == SPV_REFLECT_RESULT_SUCCESS);
-
-		const FString ResourceName(ANSI_TO_TCHAR(Binding->name));
-		CCHeaderWriter.WriteUAV(*ResourceName, UAVBindings++);
-
-		Spirv.ReflectionInfo.Add(FVulkanSpirv::FEntry(ResourceName, BindingIndex));
-	}
-
-	for (const SpvReflectDescriptorBinding* Binding : Bindings.SBufferUAVs)
-	{
-		int32 BindingIndex = MapDescriptorBindingToIndex(Binding);
-
-		SpvResult = Reflection.ChangeDescriptorBindingNumbers(Binding, BindingIndex);//, GlobalSetId);
-		check(SpvResult == SPV_REFLECT_RESULT_SUCCESS);
-
-		const FString ResourceName(ANSI_TO_TCHAR(Binding->name));
-		CCHeaderWriter.WriteUAV(*ResourceName, UAVBindings++);
-
-		Spirv.ReflectionInfo.Add(FVulkanSpirv::FEntry(ResourceName, BindingIndex));
-	}
-
-	for (const SpvReflectDescriptorBinding* Binding : Bindings.TextureUAVs)
+	// Final descriptor binding numbers for all other resource types
 	{
-		int32 BindingIndex = MapDescriptorBindingToIndex(Binding);
-
-		SpvResult = Reflection.ChangeDescriptorBindingNumbers(Binding, BindingIndex);//, GlobalSetId);
-		check(SpvResult == SPV_REFLECT_RESULT_SUCCESS);
-
-		const FString ResourceName(ANSI_TO_TCHAR(Binding->name));
-		CCHeaderWriter.WriteUAV(*ResourceName, UAVBindings++);
-
-		Spirv.ReflectionInfo.Add(FVulkanSpirv::FEntry(ResourceName, BindingIndex));
-	}
+		int32 UAVBindings = 0, SRVBindings = 0, SMPBindings = 0, ASBindings = 0;
+		auto AddReflectionInfos = [&](TArray<SpvReflectDescriptorBinding*>& BindingArray, EVulkanBindingType::EType BindingType)
+		{
+			for (const SpvReflectDescriptorBinding* Binding : BindingArray)
+			{
+				const int32 BindingIndex = BindingTable.GetRealBindingIndex(BindingToIndexMap[Binding]);
 
-	for (const SpvReflectDescriptorBinding* Binding : Bindings.TBufferSRVs)
-	{
-		int32 BindingIndex = MapDescriptorBindingToIndex(Binding);
+				SpvReflectResult SpvResult = Reflection.ChangeDescriptorBindingNumbers(Binding, BindingIndex);
+				check(SpvResult == SPV_REFLECT_RESULT_SUCCESS);
 
-		SpvResult = Reflection.ChangeDescriptorBindingNumbers(Binding, BindingIndex);//, GlobalSetId);
-		check(SpvResult == SPV_REFLECT_RESULT_SUCCESS);
+				const FString ResourceName(ANSI_TO_TCHAR(Binding->name));
 
-		const FString ResourceName(ANSI_TO_TCHAR(Binding->name));
-		CCHeaderWriter.WriteSRV(*ResourceName, SRVBindings++);
+				switch (BindingType)
+				{
+				case EVulkanBindingType::StorageTexelBuffer:
+				case EVulkanBindingType::StorageBuffer:
+				case EVulkanBindingType::StorageImage:
+					CCHeaderWriter.WriteUAV(*ResourceName, UAVBindings++);
+					break;
 
-		Spirv.ReflectionInfo.Add(FVulkanSpirv::FEntry(ResourceName, BindingIndex));
-	}
+				case EVulkanBindingType::UniformTexelBuffer:
+					CCHeaderWriter.WriteSRV(*ResourceName, SRVBindings++);
+					break;
 
-	for (const SpvReflectDescriptorBinding* Binding : Bindings.SBufferSRVs)
-	{
-		int32 BindingIndex = MapDescriptorBindingToIndex(Binding);
+				case EVulkanBindingType::Image:
+				{
+					if (Binding->usage_binding_count > 0)
+					{
+						TArray<FString> AssociatedResourceNames;
+						AssociatedResourceNames.SetNum(Binding->usage_binding_count);
 
-		SpvResult = Reflection.ChangeDescriptorBindingNumbers(Binding, BindingIndex);//, GlobalSetId);
-		check(SpvResult == SPV_REFLECT_RESULT_SUCCESS);
+						for (uint32 UsageIndex = 0; UsageIndex < Binding->usage_binding_count; ++UsageIndex)
+						{
+							const SpvReflectDescriptorBinding* AssociatedResource = Binding->usage_bindings[UsageIndex];
+							AssociatedResourceNames[UsageIndex] = ANSI_TO_TCHAR(AssociatedResource->name);
+						}
 
-		const FString ResourceName(ANSI_TO_TCHAR(Binding->name));
-		CCHeaderWriter.WriteSRV(*ResourceName, SRVBindings++);
+						CCHeaderWriter.WriteSRV(*ResourceName, SRVBindings++, /*Count:*/ 1, AssociatedResourceNames);
+					}
+					else
+					{
+						CCHeaderWriter.WriteSRV(*ResourceName, SRVBindings++);
+					}
+				}
+				break;
 
-		Spirv.ReflectionInfo.Add(FVulkanSpirv::FEntry(ResourceName, BindingIndex));
-	}
+				case EVulkanBindingType::Sampler:
+					CCHeaderWriter.WriteSamplerState(*ResourceName, SMPBindings++);
+					break;
 
-	for (const SpvReflectDescriptorBinding* Binding : Bindings.TextureSRVs)
-	{
-		int32 BindingIndex = MapDescriptorBindingToIndex(Binding);
+				case EVulkanBindingType::AccelerationStructure:
+					CCHeaderWriter.WriteAccelerationStructures(*ResourceName, ASBindings++);
+					break;
 
-		SpvResult = Reflection.ChangeDescriptorBindingNumbers(Binding, BindingIndex);//, GlobalSetId);
-		check(SpvResult == SPV_REFLECT_RESULT_SUCCESS);
+				case EVulkanBindingType::InputAttachment:
+					// Do Nothing
+					break;
 
-		const FString ResourceName(ANSI_TO_TCHAR(Binding->name));
-		if (Binding->usage_binding_count > 0)
-		{
-			TArray<FString> AssociatedResourceNames;
-			AssociatedResourceNames.SetNum(Binding->usage_binding_count);
+				default:
+					check(false);
+					break;
+				};
 
-			for (uint32 UsageIndex = 0; UsageIndex < Binding->usage_binding_count; ++UsageIndex)
-			{
-				const SpvReflectDescriptorBinding* AssociatedResource = Binding->usage_bindings[UsageIndex];
-				AssociatedResourceNames[UsageIndex] = ANSI_TO_TCHAR(AssociatedResource->name);
+				Spirv.ReflectionInfo.Add(FVulkanSpirv::FEntry(ResourceName, BindingIndex));
 			}
+		};
 
-			CCHeaderWriter.WriteSRV(*ResourceName, SRVBindings++, /*Count:*/ 1, AssociatedResourceNames);
-		}
-		else
-		{
-			CCHeaderWriter.WriteSRV(*ResourceName, SRVBindings++);
-		}
-
-		Spirv.ReflectionInfo.Add(FVulkanSpirv::FEntry(ResourceName, BindingIndex));
-	}
-
-	for (const SpvReflectDescriptorBinding* Binding : Bindings.Samplers)
-	{
-		int32 BindingIndex = MapDescriptorBindingToIndex(Binding);
-
-		SpvResult = Reflection.ChangeDescriptorBindingNumbers(Binding, BindingIndex);//, GlobalSetId);
-		check(SpvResult == SPV_REFLECT_RESULT_SUCCESS);
-
-		const FString ResourceName(ANSI_TO_TCHAR(Binding->name));
-		Spirv.ReflectionInfo.Add(FVulkanSpirv::FEntry(ResourceName, BindingIndex));
 
-		// Only emit sampler state when its shared, i.e. used with at least 2 textures
-//		if (const uint32* UseCount = SamplerStatesUseCount.Find(Binding))
-		{
-//			if (*UseCount >= 2)
-			{
-				CCHeaderWriter.WriteSamplerState(*ResourceName, SMPBindings++);
-			}
-		}
-	}
-
-	for (const SpvReflectDescriptorBinding* Binding : Bindings.AccelerationStructures)
-	{
-		int32 BindingIndex = MapDescriptorBindingToIndex(Binding);
+		AddReflectionInfos(Bindings.InputAttachments, EVulkanBindingType::InputAttachment);
 
-		SpvResult = Reflection.ChangeDescriptorBindingNumbers(Binding, BindingIndex);//, GlobalSetId);
-		check(SpvResult == SPV_REFLECT_RESULT_SUCCESS);
+		AddReflectionInfos(Bindings.TBufferUAVs, EVulkanBindingType::StorageTexelBuffer);
+		AddReflectionInfos(Bindings.SBufferUAVs, EVulkanBindingType::StorageBuffer);
+		AddReflectionInfos(Bindings.TextureUAVs, EVulkanBindingType::StorageImage);
 
-		const FString ResourceName(ANSI_TO_TCHAR(Binding->name));
-		
-		CCHeaderWriter.WriteAccelerationStructures(*ResourceName, ASBindings++);
+		AddReflectionInfos(Bindings.TBufferSRVs, EVulkanBindingType::UniformTexelBuffer);
+		AddReflectionInfos(Bindings.SBufferSRVs, EVulkanBindingType::UniformTexelBuffer);
+		AddReflectionInfos(Bindings.TextureSRVs, EVulkanBindingType::Image);
 
-		Spirv.ReflectionInfo.Add(FVulkanSpirv::FEntry(ResourceName, BindingIndex));
+		AddReflectionInfos(Bindings.Samplers, EVulkanBindingType::Sampler);
+		AddReflectionInfos(Bindings.AccelerationStructures, EVulkanBindingType::AccelerationStructure);
 	}
 
 	// Build final shader output with meta data
-	FString DebugName = Input.DumpDebugInfoPath.Right(Input.DumpDebugInfoPath.Len() - Input.DumpDebugInfoRootPath.Len());
-
-	CCHeaderWriter.WriteSourceInfo(*Input.VirtualSourceFilePath, *Input.EntryPointName);
+	CCHeaderWriter.WriteSourceInfo(*InternalState.Input.VirtualSourceFilePath, *InternalState.GetEntryPointName());
 	CCHeaderWriter.WriteCompilerInfo();
 
 	const FString MetaData = CCHeaderWriter.ToString();
 
-	Output.Target = Input.Target;
+	Output.Target = InternalState.Input.Target;
 
 	// Overwrite updated SPIRV code
 	Spirv.Data = TArray<uint32>(Reflection.GetCode(), Reflection.GetCodeSize() / 4);
 
 	// We have to strip out most debug instructions (except OpName) for Vulkan mobile
-	if (bStripReflect)
+	if (InternalState.bStripReflect)
 	{
 		const char* OptArgs[] = { "--strip-reflect", "-O"};
 		if (!CompilerContext.OptimizeSpirv(Spirv.Data, OptArgs, UE_ARRAY_COUNT(OptArgs)))
@@ -1985,7 +1980,7 @@ static bool BuildShaderOutputFromSpirv(
 	}
 
 	// For Android run an additional pass to patch spirv to be compatible across drivers
-	if(IsAndroidShaderFormat(Input.ShaderFormat))
+	if (IsAndroidShaderFormat(InternalState.Input.ShaderFormat))
 	{
 		const char* OptArgs[] = { "--android-driver-patch" };
 		if (!CompilerContext.OptimizeSpirv(Spirv.Data, OptArgs, UE_ARRAY_COUNT(OptArgs)))
@@ -2002,23 +1997,23 @@ static bool BuildShaderOutputFromSpirv(
 
 	BuildShaderOutput(
 		Output,
-		Input,
+		InternalState.Input,
 		TCHAR_TO_ANSI(*MetaData),
 		MetaData.Len(),
 		BindingTable,
 		ApproxInstructionCount,
 		WaveSize,
 		Spirv,
-		DebugName,
+		InternalState.GetDebugName(),
 		true // source contains meta data only
 	);
 
-	if (bDebugDump)
+	if (InternalState.bDebugDump)
 	{
 		// Write meta data to debug output file and write SPIR-V dump in binary and text form
-		DumpDebugShaderText(Input, MetaData, TEXT("meta.txt"));
-		DumpDebugShaderBinary(Input, Spirv.GetByteData(), Spirv.GetByteSize(), TEXT("spv"));
-		DumpDebugShaderDisassembledSpirv(Input, Spirv.GetByteData(), Spirv.GetByteSize(), TEXT("spvasm"));
+		DumpDebugShaderText(InternalState.Input, MetaData, TEXT("meta.txt"));
+		DumpDebugShaderBinary(InternalState.Input, Spirv.GetByteData(), Spirv.GetByteSize(), TEXT("spv"));
+		DumpDebugShaderDisassembledSpirv(InternalState.Input, Spirv.GetByteData(), Spirv.GetByteSize(), TEXT("spvasm"));
 	}
 
 	return true;
@@ -2151,10 +2146,10 @@ static void Patch64bitSamplers(FVulkanSpirv& Spirv)
 	}
 }
 
-static FString VulkanGetShaderProfileDXC(const FCompilerInfo& CompilerInfo, const CrossCompiler::FShaderConductorOptions Options)
+static FString VulkanGetShaderProfileDXC(const VulkanShaderCompilerInternalState& InternalState, const CrossCompiler::FShaderConductorOptions Options)
 {
 	const TCHAR* ShaderProfile = TEXT("unknown");
-	switch (CompilerInfo.Input.Target.Frequency)
+	switch (InternalState.GetShaderFrequency())
 	{
 	case SF_Vertex:
 		ShaderProfile = TEXT("vs");
@@ -2183,13 +2178,10 @@ static FString VulkanGetShaderProfileDXC(const FCompilerInfo& CompilerInfo, cons
 }
 
 static void VulkanCreateDXCCompileBatchFiles(
-	const FString& EntryPointName,
-	EShaderFrequency Frequency,
-	const FCompilerInfo& CompilerInfo,
+	const VulkanShaderCompilerInternalState& InternalState,
 	const CrossCompiler::FShaderConductorOptions Options)
 {
-
-	FString USFFilename = CompilerInfo.Input.GetSourceFilename();
+	FString USFFilename = InternalState.Input.GetSourceFilename();
 	FString SPVFilename = FPaths::GetBaseFilename(USFFilename) + TEXT(".DXC.spv");
 	FString GLSLFilename = FPaths::GetBaseFilename(USFFilename) + TEXT(".SPV.glsl");
 
@@ -2223,7 +2215,7 @@ static void VulkanCreateDXCCompileBatchFiles(
 		ensure(false);
 	}
 
-	FString ShaderProfile = VulkanGetShaderProfileDXC(CompilerInfo, Options);
+	FString ShaderProfile = VulkanGetShaderProfileDXC(InternalState, Options);
 
 	// CompileDXC.bat
 	{
@@ -2252,7 +2244,7 @@ static void VulkanCreateDXCCompileBatchFiles(
 			*DxcPath,
 			Options.HlslVersion,
 			*ShaderProfile,
-			*EntryPointName,
+			*InternalState.GetEntryPointName(),
 			VulkanVersion,
 			*SPVFilename,
 			*USFFilename,
@@ -2260,7 +2252,7 @@ static void VulkanCreateDXCCompileBatchFiles(
 			*SPVFilename
 		);
 
-		FFileHelper::SaveStringToFile(BatchFileContents, *(CompilerInfo.Input.DumpDebugInfoPath / TEXT("CompileDXC.bat")));
+		FFileHelper::SaveStringToFile(BatchFileContents, *(InternalState.Input.DumpDebugInfoPath / TEXT("CompileDXC.bat")));
 	}
 }
 
@@ -2354,74 +2346,15 @@ static FString ParseEntrypointDecl(FStringView PreprocessedShader, FStringView E
 	return EntryPointDecl;
 }
 
-static bool CompileWithShaderConductor(
-	const FString&			PreprocessedShader,
-	const FString&			EntryPointName,
-	EShaderFrequency		Frequency,
-	const FCompilerInfo&	CompilerInfo,
-	FShaderCompilerOutput&	Output,
-	FVulkanBindingTable&	BindingTable,
-	bool					bStripReflect,
-	CrossCompiler::FShaderConductorOptions::ETargetEnvironment MinTargetEnvironment)
+uint8 ParseWaveSize(
+	const VulkanShaderCompilerInternalState InternalState,
+	const FString& PreprocessedShader
+	)
 {
-	const FShaderCompilerInput& Input = CompilerInfo.Input;
-
-	const bool bIsRayTracingShader = Input.IsRayTracingShader();
-	const bool bRewriteHlslSource = !bIsRayTracingShader;
-	const bool bDebugDump = Input.DumpDebugInfoEnabled();
-
-	CrossCompiler::FShaderConductorContext CompilerContext;
-
-	// Inject additional macro definitions to circumvent missing features: external textures
-	PRAGMA_DISABLE_DEPRECATION_WARNINGS		// FShaderCompilerDefinitions will be made internal in the future, marked deprecated until then
-	FShaderCompilerDefinitions AdditionalDefines;
-	PRAGMA_ENABLE_DEPRECATION_WARNINGS
-
-	// Load shader source into compiler context
-	CompilerContext.LoadSource(PreprocessedShader, Input.VirtualSourceFilePath, EntryPointName, Frequency, &AdditionalDefines);
-
-	// Initialize compilation options for ShaderConductor
-	CrossCompiler::FShaderConductorOptions Options;
-	Options.TargetEnvironment = MinTargetEnvironment;
-
-	// VK_EXT_scalar_block_layout is required by raytracing and by Nanite (so expect it to be present in SM6/Vulkan_1_3)
-	Options.bDisableScalarBlockLayout = !(bIsRayTracingShader || 
-		MinTargetEnvironment >= CrossCompiler::FShaderConductorOptions::ETargetEnvironment::Vulkan_1_3);
-
-	if (Input.Environment.CompilerFlags.Contains(CFLAG_AllowRealTypes))
-	{
-		Options.bEnable16bitTypes = true;
-	}
-
-	// Enable HLSL 2021 if specified
-	if (Input.Environment.CompilerFlags.Contains(CFLAG_HLSL2021))
-	{
-		Options.HlslVersion = 2021;
-	}
-
-	// Ray tracing features require Vulkan 1.2 environment minimum.
-	if (bIsRayTracingShader || Input.Environment.CompilerFlags.Contains(CFLAG_InlineRayTracing))
-	{
-		if (Options.TargetEnvironment < CrossCompiler::FShaderConductorOptions::ETargetEnvironment::Vulkan_1_2)
-		{
-			Options.TargetEnvironment = CrossCompiler::FShaderConductorOptions::ETargetEnvironment::Vulkan_1_2;
-		}
-	}
-
-	if (bDebugDump)
-	{
-		VulkanCreateDXCCompileBatchFiles(
-			EntryPointName,
-			Frequency,
-			CompilerInfo,
-			Options);
-	}
-
-	// Before the shader rewritter removes all traces of it, pull any WAVESIZE directives from the shader source
 	uint8 WaveSize = 0;
-	if (!bIsRayTracingShader)
+	if (!InternalState.IsRayTracingShader())
 	{
-		const FString EntrypointDecl = ParseEntrypointDecl(PreprocessedShader, EntryPointName);
+		const FString EntrypointDecl = ParseEntrypointDecl(PreprocessedShader, InternalState.GetEntryPointName());
 
 		const FString WaveSizeMacro(TEXT("VULKAN_WAVESIZE("));
 		int32 WaveSizeIndex = EntrypointDecl.Find(*WaveSizeMacro, ESearchCase::CaseSensitive);
@@ -2457,11 +2390,64 @@ static bool CompileWithShaderConductor(
 	}
 
 	// Take note of preferred wave size flag if none was specified in HLSL
-	if ((WaveSize == 0) && Input.Environment.CompilerFlags.Contains(CFLAG_Wave32))
+	if ((WaveSize == 0) && InternalState.Input.Environment.CompilerFlags.Contains(CFLAG_Wave32))
 	{
 		WaveSize = 32;
 	}
 
+	return WaveSize;
+}
+
+static bool CompileWithShaderConductor(
+	const VulkanShaderCompilerInternalState InternalState,
+	const FString&			PreprocessedShader,
+	FShaderCompilerOutput&	Output,
+	FVulkanBindingTable&	BindingTable
+)
+{
+	const FShaderCompilerInput& Input = InternalState.Input;
+
+	const bool bIsRayTracingShader = InternalState.IsRayTracingShader();
+	const bool bRewriteHlslSource = !bIsRayTracingShader;
+	const bool bDebugDump = InternalState.bDebugDump;
+
+	CrossCompiler::FShaderConductorContext CompilerContext;
+
+	// Inject additional macro definitions to circumvent missing features: external textures
+	PRAGMA_DISABLE_DEPRECATION_WARNINGS		// FShaderCompilerDefinitions will be made internal in the future, marked deprecated until then
+	FShaderCompilerDefinitions AdditionalDefines;
+	PRAGMA_ENABLE_DEPRECATION_WARNINGS
+
+	// Load shader source into compiler context
+	CompilerContext.LoadSource(PreprocessedShader, Input.VirtualSourceFilePath, InternalState.GetEntryPointName(), InternalState.GetShaderFrequency(), &AdditionalDefines);
+
+	// Initialize compilation options for ShaderConductor
+	CrossCompiler::FShaderConductorOptions Options;
+	Options.TargetEnvironment = InternalState.MinimumTargetEnvironment;
+
+	// VK_EXT_scalar_block_layout is required by raytracing and by Nanite (so expect it to be present in SM6/Vulkan_1_3)
+	Options.bDisableScalarBlockLayout = !(bIsRayTracingShader || 
+		InternalState.MinimumTargetEnvironment >= CrossCompiler::FShaderConductorOptions::ETargetEnvironment::Vulkan_1_3);
+
+	if (Input.Environment.CompilerFlags.Contains(CFLAG_AllowRealTypes))
+	{
+		Options.bEnable16bitTypes = true;
+	}
+
+	// Enable HLSL 2021 if specified
+	if (Input.Environment.CompilerFlags.Contains(CFLAG_HLSL2021))
+	{
+		Options.HlslVersion = 2021;
+	}
+
+	if (bDebugDump)
+	{
+		VulkanCreateDXCCompileBatchFiles(InternalState, Options);
+	}
+
+	// Before the shader rewritter removes all traces of it, pull any WAVESIZE directives from the shader source
+	const uint8 WaveSize = ParseWaveSize(InternalState, PreprocessedShader);
+
 	if (bRewriteHlslSource)
 	{
 		// Rewrite HLSL source code to remove unused global resources and variables
@@ -2494,7 +2480,7 @@ static bool CompileWithShaderConductor(
 	Patch64bitSamplers(Spirv);
 
 	// Build shader output and binding table
-	Output.bSucceeded = BuildShaderOutputFromSpirv(CompilerContext, Spirv, Input, Output, BindingTable, EntryPointName, WaveSize, bStripReflect, bIsRayTracingShader, bDebugDump);
+	Output.bSucceeded = BuildShaderOutputFromSpirv(CompilerContext, Spirv, InternalState, Output, BindingTable, WaveSize);
 
 	if (Input.Environment.CompilerFlags.Contains(CFLAG_ExtraShaderData))
 	{
@@ -2511,60 +2497,27 @@ static bool CompileWithShaderConductor(
 
 
 
-EVulkanShaderVersion FormatToVersion(FName Format)
-{
-	if (Format == NAME_VULKAN_ES3_1)
-	{
-		return EVulkanShaderVersion::ES3_1;
-	}
-	else if (Format == NAME_VULKAN_ES3_1_ANDROID)
-	{
-		return EVulkanShaderVersion::ES3_1_ANDROID;
-	}
-	else if (Format == NAME_VULKAN_SM5_ANDROID)
-	{
-		return EVulkanShaderVersion::SM5_ANDROID;
-	}
-	else if (Format == NAME_VULKAN_SM5)
-	{
-		return EVulkanShaderVersion::SM5;
-	}
-	else if (Format == NAME_VULKAN_SM6)
-	{
-		return EVulkanShaderVersion::SM6;
-	}
-	else
-	{
-		FString FormatStr = Format.ToString();
-		checkf(0, TEXT("Invalid shader format passed to Vulkan shader compiler: %s"), *FormatStr);
-		return EVulkanShaderVersion::Invalid;
-	}
-}
+
 
 bool PreprocessVulkanShader(const FShaderCompilerInput& Input, const FShaderCompilerEnvironment& Environment, FShaderPreprocessOutput& PreprocessOutput)
 {
-	EVulkanShaderVersion Version = FormatToVersion(Input.ShaderFormat);
+	const VulkanShaderCompilerInternalState InternalState(Input);
 
 	PRAGMA_DISABLE_DEPRECATION_WARNINGS		// FShaderCompilerDefinitions will be made internal in the future, marked deprecated until then
 
-	const bool bIsSM6 = (Version == EVulkanShaderVersion::SM6);
-	const bool bIsSM5 = (Version == EVulkanShaderVersion::SM5) || (Version == EVulkanShaderVersion::SM5_ANDROID);
-	const bool bIsMobileES31 = (Version == EVulkanShaderVersion::ES3_1 || Version == EVulkanShaderVersion::ES3_1_ANDROID);
-	const CrossCompiler::FShaderConductorOptions::ETargetEnvironment MinTargetEnvironment = GetMinimumTargetEnvironment(Version);
-	
 	FShaderCompilerDefinitions AdditionalDefines;
 	AdditionalDefines.SetDefine(TEXT("COMPILER_HLSLCC"), 1);
 	AdditionalDefines.SetDefine(TEXT("COMPILER_VULKAN"), 1);
-	if (bIsMobileES31)
+	if (InternalState.IsMobileES31())
 	{
 		AdditionalDefines.SetDefine(TEXT("ES3_1_PROFILE"), 1);
 		AdditionalDefines.SetDefine(TEXT("VULKAN_PROFILE"), 1);
 	}
-	else if (bIsSM6)
+	else if (InternalState.IsSM6())
 	{
 		AdditionalDefines.SetDefine(TEXT("VULKAN_PROFILE_SM6"), 1);
 	}
-	else if (bIsSM5)
+	else if (InternalState.IsSM5())
 	{
 		AdditionalDefines.SetDefine(TEXT("VULKAN_PROFILE_SM5"), 1);
 	}
@@ -2589,15 +2542,11 @@ bool PreprocessVulkanShader(const FShaderCompilerInput& Input, const FShaderComp
 		AdditionalDefines.SetDefine(TEXT("PLATFORM_SUPPORTS_REAL_TYPES"), 1);
 	}
 
-	if (MinTargetEnvironment >= CrossCompiler::FShaderConductorOptions::ETargetEnvironment::Vulkan_1_1)
+	// We have ETargetEnvironment::Vulkan_1_1 by default as a min spec now
 	{
 		AdditionalDefines.SetDefine(TEXT("PLATFORM_SUPPORTS_SM6_0_WAVE_OPERATIONS"), 1);
 		AdditionalDefines.SetDefine(TEXT("VULKAN_SUPPORTS_SUBGROUP_SIZE_CONTROL"), 1);
 	}
-	else
-	{
-		check(!Input.Environment.CompilerFlags.Contains(CFLAG_WaveOperations));
-	}
 
 	AdditionalDefines.SetDefine(TEXT("VULKAN_BINDLESS_SAMPLER_ARRAY_PREFIX"), VulkanBindless::kBindlessSamplerArrayPrefix);
 	AdditionalDefines.SetDefine(TEXT("VULKAN_BINDLESS_RESOURCE_ARRAY_PREFIX"), VulkanBindless::kBindlessResourceArrayPrefix);
@@ -2658,60 +2607,32 @@ bool PreprocessVulkanShader(const FShaderCompilerInput& Input, const FShaderComp
 void CompileVulkanShader(const FShaderCompilerInput& Input, const FShaderPreprocessOutput& PreprocessOutput, FShaderCompilerOutput& Output, const class FString& WorkingDirectory)
 {
 	check(IsVulkanShaderFormat(Input.ShaderFormat));
-	EVulkanShaderVersion Version = FormatToVersion(Input.ShaderFormat);
-
-	const bool bIsSM6 = (Version == EVulkanShaderVersion::SM6);
-	const bool bIsSM5 = (Version == EVulkanShaderVersion::SM5) || (Version == EVulkanShaderVersion::SM5_ANDROID);
-	const bool bIsMobileES31 = (Version == EVulkanShaderVersion::ES3_1 || Version == EVulkanShaderVersion::ES3_1_ANDROID);
-	bool bStripReflect = Input.IsRayTracingShader();
-	// By default we strip reflecion information for Android platform to avoid issues with older drivers
-	if (IsAndroidShaderFormat(Input.ShaderFormat))
-	{
-		bStripReflect = Input.Environment.GetCompileArgument(TEXT("STRIP_REFLECT_ANDROID"), true);
-	}
 
-	const CrossCompiler::FShaderConductorOptions::ETargetEnvironment MinTargetEnvironment = GetMinimumTargetEnvironment(Version);
+	const VulkanShaderCompilerInternalState InternalState(Input);
 
-	const EHlslShaderFrequency FrequencyTable[] =
-	{
-		HSF_VertexShader,
-		HSF_InvalidFrequency,
-		HSF_InvalidFrequency,
-		HSF_PixelShader,
-		(bIsSM5 || bIsSM6) ? HSF_GeometryShader : HSF_InvalidFrequency,
-		HSF_ComputeShader, 
-		(bIsSM5 || bIsSM6) ? HSF_RayGen : HSF_InvalidFrequency,
-		(bIsSM5 || bIsSM6) ? HSF_RayMiss : HSF_InvalidFrequency,
-		(bIsSM5 || bIsSM6) ? HSF_RayHitGroup : HSF_InvalidFrequency,
-		(bIsSM5 || bIsSM6) ? HSF_RayCallable : HSF_InvalidFrequency,
-	};
-
-	const EShaderFrequency Frequency = (EShaderFrequency)Input.Target.Frequency;
-
-	const EHlslShaderFrequency HlslFrequency = FrequencyTable[Input.Target.Frequency];
+	const EHlslShaderFrequency HlslFrequency = InternalState.GetHlslShaderFrequency();
 	if (HlslFrequency == HSF_InvalidFrequency)
 	{
 		Output.bSucceeded = false;
 		FShaderCompilerError& NewError = Output.Errors.AddDefaulted_GetRef();
 		NewError.StrippedErrorMessage = FString::Printf(
 			TEXT("%s shaders not supported for use in Vulkan."),
-			CrossCompiler::GetFrequencyName(Frequency));
+			CrossCompiler::GetFrequencyName(InternalState.GetShaderFrequency()));
 		return;
 	}
 
-	const bool bDirectCompile = FParse::Param(FCommandLine::Get(), TEXT("directcompile"));
-	FCompilerInfo CompilerInfo(Input, WorkingDirectory, HlslFrequency);
+	FVulkanBindingTable BindingTable(HlslFrequency);
 
-	FVulkanBindingTable BindingTable(CompilerInfo.Frequency);
 	bool bSuccess = false;
 
 #if PLATFORM_MAC || PLATFORM_WINDOWS || PLATFORM_LINUX
 	// Cross-compile shader via ShaderConductor (DXC, SPIRV-Tools, SPIRV-Cross)
-	bSuccess = CompileWithShaderConductor(PreprocessOutput.GetSource(), Input.EntryPointName, Frequency, CompilerInfo, Output, BindingTable, bStripReflect, MinTargetEnvironment);
+	bSuccess = CompileWithShaderConductor(InternalState, PreprocessOutput.GetSource(), Output, BindingTable);
 #endif // PLATFORM_MAC || PLATFORM_WINDOWS || PLATFORM_LINUX
 	
-	PreprocessOutput.GetParameterParser().ValidateShaderParameterTypes(Input, bIsMobileES31, Output);
+	PreprocessOutput.GetParameterParser().ValidateShaderParameterTypes(Input, InternalState.IsMobileES31(), Output);
 	
+	const bool bDirectCompile = FParse::Param(FCommandLine::Get(), TEXT("directcompile"));
 	if (bDirectCompile)
 	{
 		for (const auto& Error : Output.Errors)
diff --git a/Engine/Source/Developer/VulkanShaderFormat/Private/VulkanShaderFormat.h b/Engine/Source/Developer/VulkanShaderFormat/Private/VulkanShaderFormat.h
--- a/Engine/Source/Developer/VulkanShaderFormat/Private/VulkanShaderFormat.h
+++ b/Engine/Source/Developer/VulkanShaderFormat/Private/VulkanShaderFormat.h
@@ -12,18 +12,6 @@
 #define UE_VULKAN_SHADER_COMPILER_ALLOW_DEAD_CODE_REMOVAL 1
 #endif // UE_VULKAN_SHADER_COMPILER_ALLOW_DEAD_CODE_REMOVAL
 
-// Hold information to be able to call the compilers
-struct FCompilerInfo
-{
-	const struct FShaderCompilerInput& Input;
-	FString WorkingDirectory;
-	FString Profile;
-	EHlslShaderFrequency Frequency;
-	FString BaseSourceFilename;
-
-	FCompilerInfo(const struct FShaderCompilerInput& InInput, const FString& InWorkingDirectory, EHlslShaderFrequency InFrequency);
-};
-
 
 struct FVulkanSpirv : FSpirv
 {
