VulkanRHI: Remove unnecessary code in VulkanShaderCompiler (clean up WIP).  Remove some of the CombinedImageSampler code (and some ImmutableSamplers).  Removed some parts of FOLDVulkanCodeHeader.  Remove the use of FHlslccHeaderWriter and FHlslccHeader to generate a string and then parse it, work directly from SPIRV and reflection.

#rnx

diff --git a/Engine/Source/Developer/VulkanShaderFormat/Private/VulkanShaderCompiler.cpp b/Engine/Source/Developer/VulkanShaderFormat/Private/VulkanShaderCompiler.cpp
--- a/Engine/Source/Developer/VulkanShaderFormat/Private/VulkanShaderCompiler.cpp
+++ b/Engine/Source/Developer/VulkanShaderFormat/Private/VulkanShaderCompiler.cpp
@@ -3,7 +3,6 @@
 
 #include "VulkanShaderFormat.h"
 
-#include "HlslccHeaderWriter.h"
 #include "hlslcc.h"
 #include "RHIShaderFormatDefinitions.inl"
 #include "ShaderCompilerCommon.h"
@@ -220,43 +219,6 @@ static void GetResourceEntryFromUBMember(const FShaderResourceTableMap& Resource
 	check(0);
 }
 
-static FString FindTextureNameForSamplerState(const CrossCompiler::FHlslccHeader& CCHeader, const FString& InSamplerName)
-{
-	for (const auto& Sampler : CCHeader.Samplers)
-	{
-		for (const auto& SamplerState : Sampler.SamplerStates)
-		{
-			if (SamplerState == InSamplerName)
-			{
-				return Sampler.Name;
-			}
-		}
-	}
-	return TEXT("");
-}
-
-static uint16 GetCombinedSamplerStateAlias(const FString& ParameterName,
-											VkDescriptorType DescriptorType,
-											const FVulkanBindingTable& BindingTable,
-											const CrossCompiler::FHlslccHeader& CCHeader,
-											const TArray<FString>& GlobalNames)
-{
-	if (DescriptorType == VK_DESCRIPTOR_TYPE_COMBINED_IMAGE_SAMPLER)
-	{
-		if (!ContainsBinding(BindingTable, ParameterName))
-		{
-			// Not found as a regular parameter, find corresponding Texture and return that ResourceEntryIndex
-			const FString& TextureName = FindTextureNameForSamplerState(CCHeader, ParameterName);
-			check(TextureName.Len() > 0);
-
-			int32 Found = GlobalNames.Find(TextureName);
-			check(Found >= 0);
-			return (uint16)Found;
-		}
-	}
-
-	return UINT16_MAX;
-}
 
 struct FPatchType
 {
@@ -265,61 +227,10 @@ struct FPatchType
 };
 
 
-class FVulkanShaderSerializedBindings : public CrossCompiler::FShaderBindings
-{
-public:
-	FVulkanShaderSerializedBindings()
-	{
-		NumSamplers = 0;
-		NumUniformBuffers = 0;
-		NumUAVs = 0;
-		NumAccelerationStructures = 0;
-		bHasRegularUniformBuffers = 0;
-	}
-};
-
-struct FOLDVulkanCodeHeader
-{
-	FVulkanShaderSerializedBindings SerializedBindings;
-
-	struct FShaderDescriptorInfo
-	{
-		TArray<VkDescriptorType> DescriptorTypes;
-		uint16 NumImageInfos;
-		uint16 NumBufferInfos;
-	};
-	FShaderDescriptorInfo NEWDescriptorInfo;
-
-	struct FPackedUBToVulkanBindingIndex
-	{
-		CrossCompiler::EPackedTypeName	TypeName;
-		uint8							VulkanBindingIndex;
-	};
-	TArray<FPackedUBToVulkanBindingIndex> NEWPackedUBToVulkanBindingIndices;
-
-	// List of memory copies from RHIUniformBuffer to packed uniforms when emulating UB's
-	TArray<CrossCompiler::FUniformBufferCopyInfo> UniformBuffersCopyInfo;
-
-	FString ShaderName;
-	FSHAHash SourceHash;
-
-	uint64 UniformBuffersWithDescriptorMask;
-
-	// Number of uniform buffers (not including PackedGlobalUBs) UNUSED
-	uint32 UNUSED_NumNonGlobalUBs;
-
-	// (Separated to improve cache) if this is non-zero, then we can assume all UBs are emulated
-	TArray<uint32> NEWPackedGlobalUBSizes;
-
-	// Number of copies per emulated buffer source index (to skip searching among UniformBuffersCopyInfo). Upper uint16 is the index, Lower uint16 is the count
-	TArray<uint32> NEWEmulatedUBCopyRanges;
-};
-
-
 // A collection of states and data that is locked in at the top level call and doesn't change throughout the compilation process
-struct VulkanShaderCompilerInternalState
+struct FVulkanShaderCompilerInternalState
 {
-	VulkanShaderCompilerInternalState(const FShaderCompilerInput& InInput)
+	FVulkanShaderCompilerInternalState(const FShaderCompilerInput& InInput)
 		: Input(InInput)
 		, Version(FormatToVersion(Input.ShaderFormat))
 		, MinimumTargetEnvironment(GetMinimumTargetEnvironment(InInput))
@@ -395,30 +306,20 @@ struct VulkanShaderCompilerInternalState
 
 
 
-static void AddImmutable(FVulkanShaderHeader& OutHeader, int32 GlobalIndex)
-{
-	check(GlobalIndex < UINT16_MAX);
-	OutHeader.Globals[GlobalIndex].bImmutableSampler = true;
-}
-
-static int32 AddGlobal(FOLDVulkanCodeHeader& OLDHeader,
-						const FVulkanBindingTable& BindingTable,
-						const CrossCompiler::FHlslccHeader& CCHeader,
+static int32 AddGlobal( const TArray<VkDescriptorType>& DescriptorTypes,
 						const FString& ParameterName,
 						uint16 BindingIndex,
 						const FVulkanSpirv& Spirv,
 						FVulkanShaderHeader& OutHeader,
-						const TArray<FString>& GlobalNames,
-						TArray<FPatchType>& OutTypePatch,
-						uint16 CombinedAliasIndex)
+						const TArray<FString>& GlobalNames
+)
 {
-	int32 HeaderGlobalIndex = GlobalNames.Find(ParameterName);//OutHeader.Globals.AddZeroed();
+	const int32 HeaderGlobalIndex = GlobalNames.Find(ParameterName);
 	check(HeaderGlobalIndex != INDEX_NONE);
 	check(GlobalNames[HeaderGlobalIndex] == ParameterName);
 
 	FVulkanShaderHeader::FGlobalInfo& GlobalInfo = OutHeader.Globals[HeaderGlobalIndex];
 	const FVulkanSpirv::FEntry* Entry = Spirv.GetEntry(ParameterName);
-	bool bIsCombinedSampler = false;
 	if (Entry)
 	{
 		if (Entry->Binding == -1)
@@ -431,41 +332,24 @@ static int32 AddGlobal(FOLDVulkanCodeHeader& OLDHeader,
 	}
 	else
 	{
-		Entry = CombinedAliasIndex == UINT16_MAX ? Spirv.GetEntryByBindingIndex(BindingIndex) : Spirv.GetEntry(GlobalNames[CombinedAliasIndex]);
+		Entry = Spirv.GetEntryByBindingIndex(BindingIndex);
 		check(Entry);
 		check(Entry->Binding != -1);
 		if (!Entry->Name.EndsWith(TEXT("_BUFFER")))
 		{
-			bIsCombinedSampler = true;
+			checkf(false, TEXT("CombinedSamplers should not be used anymore"))
 		}
 	}
 
-	VkDescriptorType DescriptorType = bIsCombinedSampler ? VK_DESCRIPTOR_TYPE_COMBINED_IMAGE_SAMPLER : OLDHeader.NEWDescriptorInfo.DescriptorTypes[Entry->Binding];
+	const VkDescriptorType DescriptorType = DescriptorTypes[Entry->Binding];
 
 	GlobalInfo.OriginalBindingIndex = Entry->Binding;
 	OutHeader.GlobalSpirvInfos[HeaderGlobalIndex] = FVulkanShaderHeader::FSpirvInfo(Entry->WordDescriptorSetIndex, Entry->WordBindingIndex);
-	if (bIsCombinedSampler)
-	{
-		uint16 NewCombinedAliasIndex = GetCombinedSamplerStateAlias(ParameterName, DescriptorType, BindingTable, CCHeader, GlobalNames);
-		check(NewCombinedAliasIndex != UINT16_MAX);
 
-		{
-			// Ideally we would set up the type index here, but we might not have processed the aliased texture yet:
-			//		GlobalInfo.TypeIndex = OutHeader.Globals[NewCombinedAliasIndex].TypeIndex;
-			// Instead postpone this patching
-			GlobalInfo.TypeIndex = UINT16_MAX;
-			OutTypePatch.Add({HeaderGlobalIndex, NewCombinedAliasIndex});
-		}
+	const int32 GlobalDescriptorTypeIndex = OutHeader.GlobalDescriptorTypes.Add(DescriptorTypeToBinding(DescriptorType));
+	GlobalInfo.TypeIndex = GlobalDescriptorTypeIndex;
+	GlobalInfo.CombinedSamplerStateAliasIndex = UINT16_MAX;
 
-		GlobalInfo.CombinedSamplerStateAliasIndex = CombinedAliasIndex == UINT16_MAX ? NewCombinedAliasIndex : CombinedAliasIndex;
-	}
-	else
-	{
-		int32 GlobalDescriptorTypeIndex = OutHeader.GlobalDescriptorTypes.Add(DescriptorTypeToBinding(DescriptorType));
-		GlobalInfo.TypeIndex = GlobalDescriptorTypeIndex;
-		check(GetCombinedSamplerStateAlias(ParameterName, DescriptorType, BindingTable, CCHeader, GlobalNames) == UINT16_MAX);
-		GlobalInfo.CombinedSamplerStateAliasIndex = UINT16_MAX;
-	}
 #if VULKAN_ENABLE_BINDING_DEBUG_NAMES
 	GlobalInfo.DebugName = ParameterName;
 #endif
@@ -473,49 +357,14 @@ static int32 AddGlobal(FOLDVulkanCodeHeader& OLDHeader,
 	return HeaderGlobalIndex;
 }
 
-static int32 AddGlobalForUBEntry(FOLDVulkanCodeHeader& OLDHeader,
-									const FVulkanBindingTable& BindingTable,
-									const CrossCompiler::FHlslccHeader& CCHeader,
-									const FString& ParameterName,
-									uint16 BindingIndex,
-									const FVulkanSpirv& Spirv,
-									const TArray<FString>&
-									GlobalNames,
-									EUniformBufferBaseType UBEntryType,
-									TArray<FPatchType>& OutTypePatch,
-									FVulkanShaderHeader& OutHeader)
-{
-	uint16 CombinedAliasIndex = UINT16_MAX;
-	if (UBEntryType == UBMT_SAMPLER)
-	{
-		if (!ContainsBinding(BindingTable, ParameterName))
-		{
-			// Not found as a regular parameter, find corresponding Texture and return that ResourceEntryIndex
-			const FString& TextureName = FindTextureNameForSamplerState(CCHeader, ParameterName);
-			check(TextureName.Len() > 0);
-
-			int32 TextureGlobalIndex = GlobalNames.Find(TextureName);
-			check(TextureGlobalIndex >= 0);
-
-			CombinedAliasIndex = (uint16)TextureGlobalIndex;
-		}
-	}
-
-	return AddGlobal(OLDHeader, BindingTable, CCHeader, ParameterName, BindingIndex, Spirv, OutHeader, GlobalNames, OutTypePatch, CombinedAliasIndex);
-}
-
-static void AddUBResources(FOLDVulkanCodeHeader& OLDHeader,
-							const FString& UBName,
+static void AddUBResources( const FString& UBName,
 							const FShaderResourceTableMap& ResourceTableMap,
 							uint32 BufferIndex,
 							const TArray<uint32>& BindingArray,
-							const FVulkanBindingTable& BindingTable,
 							const TArray<VkDescriptorType>& DescriptorTypes,
 							const FVulkanSpirv& Spirv,
-							const CrossCompiler::FHlslccHeader& CCHeader,
 							FVulkanShaderHeader::FUniformBufferInfo& OutUBInfo,
 							FVulkanShaderHeader& OutHeader,
-							TArray<FPatchType>& OutTypePatch,
 							TArray<FString>& GlobalNames)
 {
 	if (BindingArray.Num() > 0)
@@ -540,10 +389,9 @@ static void AddUBResources(FOLDVulkanCodeHeader& OLDHeader,
 				FUniformResourceEntry ResourceTableEntry;
 				GetResourceEntryFromUBMember(ResourceTableMap, UBName, ResourceIndex, ResourceTableEntry);
 
-				int32 HeaderUBResourceInfoIndex = OutUBInfo.ResourceEntries.AddZeroed();
-				FVulkanShaderHeader::FUBResourceInfo& UBResourceInfo = OutUBInfo.ResourceEntries[HeaderUBResourceInfoIndex];
+				FVulkanShaderHeader::FUBResourceInfo& UBResourceInfo = OutUBInfo.ResourceEntries.AddZeroed_GetRef();;
 
-				int32 HeaderGlobalIndex = AddGlobalForUBEntry(OLDHeader, BindingTable, CCHeader, ResourceTableEntry.UniformBufferMemberName, BindingIndex, Spirv, GlobalNames, (EUniformBufferBaseType)ResourceTableEntry.Type, OutTypePatch, OutHeader);
+				const int32 HeaderGlobalIndex = AddGlobal(DescriptorTypes, ResourceTableEntry.UniformBufferMemberName, BindingIndex, Spirv, OutHeader, GlobalNames);
 				UBResourceInfo.SourceUBResourceIndex = ResourceIndex;
 				UBResourceInfo.OriginalBindingIndex = BindingIndex;
 				UBResourceInfo.GlobalIndex = HeaderGlobalIndex;
@@ -559,17 +407,17 @@ static void AddUBResources(FOLDVulkanCodeHeader& OLDHeader,
 	}
 }
 
-static void AddUniformBuffer(FOLDVulkanCodeHeader& OLDHeader,
-	const FVulkanBindingTable& BindingTable,
+static void AddUniformBuffer(
+	const FShaderCompilerResourceTable& ShaderResourceTable,
+	const TArray<VkDescriptorType>& DescriptorTypes,
 	const FShaderCompilerInput& ShaderInput,
-	const CrossCompiler::FHlslccHeader& CCHeader,
 	const FVulkanSpirv& Spirv,
 	const FString& UBName,
 	uint16 BindingIndex,
 	FShaderParameterMap& InOutParameterMap,
 	FVulkanShaderHeader& OutHeader,
-	TArray<FPatchType>& OutTypePatch,
-	TArray<FString>& GlobalNames)
+	TArray<FString>& GlobalNames
+)
 {
 	const int32 HeaderUBIndex = OutHeader.UniformBuffers.AddZeroed();
 	FVulkanShaderHeader::FUniformBufferInfo& UBInfo = OutHeader.UniformBuffers[HeaderUBIndex];
@@ -611,14 +459,13 @@ static void AddUniformBuffer(FOLDVulkanCodeHeader& OLDHeader,
 	}
 
 	// Add used resources...
-	const FShaderCompilerResourceTable& SRT = OLDHeader.SerializedBindings.ShaderResourceTable;
-	if (SRT.ResourceTableBits & (1 << BindingIndex))
+	if (ShaderResourceTable.ResourceTableBits & (1 << BindingIndex))
 	{
 		// Make sure to process in the same order as when gathering names below
-		AddUBResources(OLDHeader, UBName, ShaderInput.Environment.ResourceTableMap, BindingIndex, SRT.TextureMap, BindingTable, OLDHeader.NEWDescriptorInfo.DescriptorTypes, Spirv, CCHeader, UBInfo, OutHeader, OutTypePatch, GlobalNames);
-		AddUBResources(OLDHeader, UBName, ShaderInput.Environment.ResourceTableMap, BindingIndex, SRT.SamplerMap, BindingTable, OLDHeader.NEWDescriptorInfo.DescriptorTypes, Spirv, CCHeader, UBInfo, OutHeader, OutTypePatch, GlobalNames);
-		AddUBResources(OLDHeader, UBName, ShaderInput.Environment.ResourceTableMap, BindingIndex, SRT.ShaderResourceViewMap, BindingTable, OLDHeader.NEWDescriptorInfo.DescriptorTypes, Spirv, CCHeader, UBInfo, OutHeader, OutTypePatch, GlobalNames);
-		AddUBResources(OLDHeader, UBName, ShaderInput.Environment.ResourceTableMap, BindingIndex, SRT.UnorderedAccessViewMap, BindingTable, OLDHeader.NEWDescriptorInfo.DescriptorTypes, Spirv, CCHeader, UBInfo, OutHeader, OutTypePatch, GlobalNames);
+		AddUBResources(UBName, ShaderInput.Environment.ResourceTableMap, BindingIndex, ShaderResourceTable.TextureMap, DescriptorTypes, Spirv, UBInfo, OutHeader, GlobalNames);
+		AddUBResources(UBName, ShaderInput.Environment.ResourceTableMap, BindingIndex, ShaderResourceTable.SamplerMap, DescriptorTypes, Spirv, UBInfo, OutHeader, GlobalNames);
+		AddUBResources(UBName, ShaderInput.Environment.ResourceTableMap, BindingIndex, ShaderResourceTable.ShaderResourceViewMap, DescriptorTypes, Spirv, UBInfo, OutHeader, GlobalNames);
+		AddUBResources(UBName, ShaderInput.Environment.ResourceTableMap, BindingIndex, ShaderResourceTable.UnorderedAccessViewMap, DescriptorTypes, Spirv, UBInfo, OutHeader, GlobalNames);
 	}
 	else
 	{
@@ -636,50 +483,15 @@ static void AddUniformBuffer(FOLDVulkanCodeHeader& OLDHeader,
 static int32 DoAddGlobal(const FString& Name, FVulkanShaderHeader& OutHeader, TArray<FString>& OutGlobalNames)
 {
 	check(!OutGlobalNames.Contains(Name));
-	int32 NameIndex = OutGlobalNames.Add(Name);
-	int32 GlobalIndex = OutHeader.Globals.AddDefaulted();
+	const int32 NameIndex = OutGlobalNames.Add(Name);
+	const int32 GlobalIndex = OutHeader.Globals.AddDefaulted();
 	check(NameIndex == GlobalIndex);
-	int32 GlobalSpirvIndex = OutHeader.GlobalSpirvInfos.AddDefaulted();
+	const int32 GlobalSpirvIndex = OutHeader.GlobalSpirvInfos.AddDefaulted();
 	check(GlobalSpirvIndex == GlobalIndex);
 	return GlobalIndex;
 }
 
-struct FVulkanHlslccHeader : public CrossCompiler::FHlslccHeader
-{
-	virtual bool ParseCustomHeaderEntries(const ANSICHAR*& ShaderSource) override
-	{
-		if (FCStringAnsi::Strncmp(ShaderSource, "// @ExternalTextures: ", 22) == 0)
-		{
-			ShaderSource += 22;
-			while (*ShaderSource && *ShaderSource != '\n')
-			{
-				FString ExternalTextureName;
-				if (!CrossCompiler::ParseIdentifier(ShaderSource, ExternalTextureName))
-				{
-					return false;
-				}
-
-				ExternalTextures.Add(ExternalTextureName);
-
-				if (Match(ShaderSource, '\n'))
-				{
-					break;
-				}
-
-				if (Match(ShaderSource, ','))
-				{
-					continue;
-				}
-			}
-		}
-
-		return true;
-	}
-
-	TArray<FString> ExternalTextures;
-};
-
-static void PrepareUBResourceEntryGlobals(const FVulkanHlslccHeader& CCHeader, const TArray<uint32>& BindingArray, const FShaderResourceTableMap& ResourceTableMap,
+static void PrepareUBResourceEntryGlobals(const TArray<uint32>& BindingArray, const FShaderResourceTableMap& ResourceTableMap,
 	int32 BufferIndex, const FString& UBName, TArray<FString>& OutGlobalNames, FVulkanShaderHeader& OutHeader)
 {
 	if (BindingArray.Num() > 0)
@@ -695,20 +507,13 @@ static void PrepareUBResourceEntryGlobals(const FVulkanHlslccHeader& CCHeader, c
 				// Verify that we have correct buffer index
 				check(FRHIResourceTableEntry::GetUniformBufferIndex(ResourceInfo) == BufferIndex);
 
-				// Extract binding index from ResourceInfo
-				const uint32 BindingIndex = FRHIResourceTableEntry::GetBindIndex(ResourceInfo);
-
 				// Extract index of the resource stored in the resource table from ResourceInfo
 				const uint16 ResourceIndex = FRHIResourceTableEntry::GetResourceIndex(ResourceInfo);
 
 				FUniformResourceEntry ResourceTableEntry;
 				GetResourceEntryFromUBMember(ResourceTableMap, UBName, ResourceIndex, ResourceTableEntry);
 
-				int32 GlobalIndex = DoAddGlobal(ResourceTableEntry.UniformBufferMemberName, OutHeader, OutGlobalNames);
-				if (CCHeader.ExternalTextures.Contains(ResourceTableEntry.UniformBufferMemberName))
-				{
-					AddImmutable(OutHeader, GlobalIndex);
-				}
+				DoAddGlobal(ResourceTableEntry.UniformBufferMemberName, OutHeader, OutGlobalNames);
 
 				// Iterate to next info
 				ResourceInfo = *ResourceInfos++;
@@ -718,21 +523,21 @@ static void PrepareUBResourceEntryGlobals(const FVulkanHlslccHeader& CCHeader, c
 	}
 }
 
-static bool IsSamplerState(const CrossCompiler::FHlslccHeader& CCHeader, const FString& ParameterName)
+
+static void PrepareGlobals(const FVulkanBindingTable& BindingTable, const FSpirvReflectBindings& SpirvReflectBindings, const FShaderCompilerResourceTable& SRT, const TMap<FString, FVulkanShaderHeader::EType>& EntryTypes, const FShaderCompilerInput& ShaderInput, const TArray<FString>& ParameterNames, FShaderParameterMap& ParameterMap, TArray<FString>& OutGlobalNames, FVulkanShaderHeader& OutHeader)
 {
-	for (const auto& Sampler : CCHeader.Samplers)
+	auto IsSamplerState = [&SpirvReflectBindings](const FString& ParameterName)
 	{
-		if (Sampler.SamplerStates.Contains(ParameterName))
+		for (SpvReflectDescriptorBinding* DescriptorBinding : SpirvReflectBindings.Samplers)
 		{
-			return true;
+			if (ParameterName == DescriptorBinding->name)
+			{
+				return true;
+			}
 		}
-	}
-
-	return false;
-}
+		return false;
+	};
 
-static void PrepareGlobals(const FVulkanBindingTable& BindingTable, const FVulkanHlslccHeader& CCHeader, const FShaderCompilerResourceTable& SRT, const TMap<FString, FVulkanShaderHeader::EType>& EntryTypes, const FShaderCompilerInput& ShaderInput, const TArray<FString>& ParameterNames, FShaderParameterMap& ParameterMap, TArray<FString>& OutGlobalNames, FVulkanShaderHeader& OutHeader)
-{
 	// First pass, gather names for all the Globals that are NOT Samplers
 	{
 		auto AddGlobalNamesForUB = [&](const FString& ParameterName)
@@ -743,28 +548,24 @@ static void PrepareGlobals(const FVulkanBindingTable& BindingTable, const FVulka
 			// Add used resources...
 			if (SRT.ResourceTableBits & (1 << ParameterAllocation->BufferIndex))
 			{
-				PrepareUBResourceEntryGlobals(CCHeader, SRT.TextureMap, ShaderInput.Environment.ResourceTableMap, ParameterAllocation->BufferIndex, ParameterName, OutGlobalNames, OutHeader);
-				PrepareUBResourceEntryGlobals(CCHeader, SRT.ShaderResourceViewMap, ShaderInput.Environment.ResourceTableMap, ParameterAllocation->BufferIndex, ParameterName, OutGlobalNames, OutHeader);
-				PrepareUBResourceEntryGlobals(CCHeader, SRT.UnorderedAccessViewMap, ShaderInput.Environment.ResourceTableMap, ParameterAllocation->BufferIndex, ParameterName, OutGlobalNames, OutHeader);
+				PrepareUBResourceEntryGlobals(SRT.TextureMap, ShaderInput.Environment.ResourceTableMap, ParameterAllocation->BufferIndex, ParameterName, OutGlobalNames, OutHeader);
+				PrepareUBResourceEntryGlobals(SRT.ShaderResourceViewMap, ShaderInput.Environment.ResourceTableMap, ParameterAllocation->BufferIndex, ParameterName, OutGlobalNames, OutHeader);
+				PrepareUBResourceEntryGlobals(SRT.UnorderedAccessViewMap, ShaderInput.Environment.ResourceTableMap, ParameterAllocation->BufferIndex, ParameterName, OutGlobalNames, OutHeader);
 			}
 		};
 
 		for (int32 ParameterIndex = 0; ParameterIndex < ParameterNames.Num(); ++ParameterIndex)
 		{
-			const FString& ParameterName = *ParameterNames[ParameterIndex];
+			const FString& ParameterName = ParameterNames[ParameterIndex];
 			const FVulkanShaderHeader::EType* FoundType = EntryTypes.Find(ParameterName);
 			if (FoundType)
 			{
 				switch (*FoundType)
 				{
 				case FVulkanShaderHeader::Global:
-					if (!IsSamplerState(CCHeader, ParameterName))
+					if (!IsSamplerState(ParameterName))
 					{
-						int32 GlobalIndex = DoAddGlobal(ParameterName, OutHeader, OutGlobalNames);
-						if (CCHeader.ExternalTextures.Contains(ParameterName))
-						{
-							AddImmutable(OutHeader, GlobalIndex);
-						}
+						DoAddGlobal(ParameterName, OutHeader, OutGlobalNames);
 					}
 					break;
 				case FVulkanShaderHeader::UniformBuffer:
@@ -795,26 +596,22 @@ static void PrepareGlobals(const FVulkanBindingTable& BindingTable, const FVulka
 			// Add used resources...
 			if (SRT.ResourceTableBits & (1 << ParameterAllocation->BufferIndex))
 			{
-				PrepareUBResourceEntryGlobals(CCHeader, SRT.SamplerMap, ShaderInput.Environment.ResourceTableMap, ParameterAllocation->BufferIndex, ParameterName, OutGlobalNames, OutHeader);
+				PrepareUBResourceEntryGlobals(SRT.SamplerMap, ShaderInput.Environment.ResourceTableMap, ParameterAllocation->BufferIndex, ParameterName, OutGlobalNames, OutHeader);
 			}
 		};
 
 		for (int32 ParameterIndex = 0; ParameterIndex < ParameterNames.Num(); ++ParameterIndex)
 		{
-			const FString& ParameterName = *ParameterNames[ParameterIndex];
+			const FString& ParameterName = ParameterNames[ParameterIndex];
 			const FVulkanShaderHeader::EType* FoundType = EntryTypes.Find(ParameterName);
 			if (FoundType)
 			{
 				switch (*FoundType)
 				{
 				case FVulkanShaderHeader::Global:
-					if (IsSamplerState(CCHeader, ParameterName))
+					if (IsSamplerState(ParameterName))
 					{
-						int32 GlobalIndex = DoAddGlobal(ParameterName, OutHeader, OutGlobalNames);
-						if (CCHeader.ExternalTextures.Contains(ParameterName))
-						{
-							AddImmutable(OutHeader, GlobalIndex);
-						}
+						DoAddGlobal(ParameterName, OutHeader, OutGlobalNames);
 					}
 					break;
 				case FVulkanShaderHeader::UniformBuffer:
@@ -848,26 +645,24 @@ static void PrepareGlobals(const FVulkanBindingTable& BindingTable, const FVulka
 	}
 }
 
-static void ConvertToNEWHeader(FOLDVulkanCodeHeader& OLDHeader,
+static void ConvertToHeader(
+	FShaderCompilerResourceTable& ShaderResourceTable,
 	const FVulkanBindingTable& BindingTable,
+	const TArray<VkDescriptorType>& DescriptorTypes,
 	const FVulkanSpirv& Spirv,
 	const TMap<FString, FVulkanShaderHeader::EType>& EntryTypes,
 	const FShaderCompilerInput& ShaderInput,
-	FVulkanHlslccHeader& CCHeader,
+	const FSpirvReflectBindings& SpirvReflectBindings,
 	FShaderParameterMap& InOutParameterMap,
 	FVulkanShaderHeader& OutHeader)
 {
 	// Names that match the Header.Globals array
 	TArray<FString> GlobalNames;
 
-	TArray<FPatchType> TypePatchList;
-
 	TArray<FString> ParameterNames;
 	InOutParameterMap.GetAllParameterNames(ParameterNames);
 
-	const FShaderCompilerResourceTable& SRT = OLDHeader.SerializedBindings.ShaderResourceTable;
-
-	PrepareGlobals(BindingTable, CCHeader, SRT, EntryTypes, ShaderInput, ParameterNames, InOutParameterMap, GlobalNames, OutHeader);
+	PrepareGlobals(BindingTable, SpirvReflectBindings, ShaderResourceTable, EntryTypes, ShaderInput, ParameterNames, InOutParameterMap, GlobalNames, OutHeader);
 
 	for (int32 ParameterIndex = 0; ParameterIndex < ParameterNames.Num(); ++ParameterIndex)
 	{
@@ -875,7 +670,8 @@ static void ConvertToNEWHeader(FOLDVulkanCodeHeader& OLDHeader,
 		uint16 BaseIndex;
 		uint16 Size;
 		const FString& ParameterName = *ParameterNames[ParameterIndex];
-		InOutParameterMap.FindParameterAllocation(*ParameterName, BufferIndex, BaseIndex, Size);
+		const bool bFoundParam = InOutParameterMap.FindParameterAllocation(*ParameterName, BufferIndex, BaseIndex, Size);
+		check(bFoundParam);
 
 		const FVulkanShaderHeader::EType* FoundType = EntryTypes.Find(ParameterName);
 		if (FoundType)
@@ -884,7 +680,7 @@ static void ConvertToNEWHeader(FOLDVulkanCodeHeader& OLDHeader,
 			{
 			case FVulkanShaderHeader::Global:
 				{
-					int32 HeaderGlobalIndex = AddGlobal(OLDHeader, BindingTable, CCHeader, ParameterName, BaseIndex, Spirv, OutHeader, GlobalNames, TypePatchList, UINT16_MAX);
+					const int32 HeaderGlobalIndex = AddGlobal(DescriptorTypes, ParameterName, BaseIndex, Spirv, OutHeader, GlobalNames);
 
 					const FParameterAllocation* ParameterAllocation = InOutParameterMap.GetParameterMap().Find(*ParameterName);
 					check(ParameterAllocation);
@@ -896,9 +692,8 @@ static void ConvertToNEWHeader(FOLDVulkanCodeHeader& OLDHeader,
 				break;
 			case FVulkanShaderHeader::PackedGlobal:
 				{
-					int32 HeaderPackedGlobalIndex = OutHeader.PackedGlobals.AddZeroed();
-					FVulkanShaderHeader::FPackedGlobalInfo& PackedGlobalInfo = OutHeader.PackedGlobals[HeaderPackedGlobalIndex];
-					PackedGlobalInfo.PackedTypeIndex = CrossCompiler::PackedTypeNameToTypeIndex(OLDHeader.NEWPackedUBToVulkanBindingIndices[BufferIndex].TypeName);
+					FVulkanShaderHeader::FPackedGlobalInfo& PackedGlobalInfo = OutHeader.PackedGlobals.AddZeroed_GetRef();
+					PackedGlobalInfo.PackedTypeIndex = CrossCompiler::EPackedTypeIndex::HighP;
 					PackedGlobalInfo.PackedUBIndex = BufferIndex;
 					checkf(Size > 0, TEXT("Assertion failed for shader parameter: %s"), *ParameterName);
 					PackedGlobalInfo.ConstantDataSizeInFloats = Size / sizeof(float);
@@ -909,7 +704,7 @@ static void ConvertToNEWHeader(FOLDVulkanCodeHeader& OLDHeader,
 				}
 				break;
 			case FVulkanShaderHeader::UniformBuffer:
-				AddUniformBuffer(OLDHeader, BindingTable, ShaderInput, CCHeader, Spirv, ParameterName, BufferIndex, InOutParameterMap, OutHeader, TypePatchList, GlobalNames);
+				AddUniformBuffer(ShaderResourceTable, DescriptorTypes, ShaderInput, Spirv, ParameterName, BufferIndex, InOutParameterMap, OutHeader, GlobalNames);
 				break;
 			default:
 				check(0);
@@ -919,38 +714,10 @@ static void ConvertToNEWHeader(FOLDVulkanCodeHeader& OLDHeader,
 		else
 		{
 			// Not found means it's a new resource-only UniformBuffer
-			AddUniformBuffer(OLDHeader, BindingTable, ShaderInput, CCHeader, Spirv, ParameterName, BufferIndex, InOutParameterMap, OutHeader, TypePatchList, GlobalNames);
+			AddUniformBuffer(ShaderResourceTable, DescriptorTypes, ShaderInput, Spirv, ParameterName, BufferIndex, InOutParameterMap, OutHeader, GlobalNames);
 		}
 	}
 
-	// Process the type patch list
-	for (const FPatchType& Patch : TypePatchList)
-	{
-		check(OutHeader.Globals[Patch.HeaderGlobalIndex].TypeIndex == UINT16_MAX);
-		OutHeader.Globals[Patch.HeaderGlobalIndex].TypeIndex = OutHeader.Globals[Patch.CombinedAliasIndex].TypeIndex;
-	}
-
-	// Add the packed global UBs
-	const FString UBOGlobalsNameSpv = ANSI_TO_TCHAR(CrossCompiler::FShaderConductorContext::GetIdentifierTable().GlobalsUniformBuffer);
-
-	for (int32 Index = 0; Index < OLDHeader.NEWPackedUBToVulkanBindingIndices.Num(); ++Index)
-	{
-		const FOLDVulkanCodeHeader::FPackedUBToVulkanBindingIndex& PackedArrayInfo = OLDHeader.NEWPackedUBToVulkanBindingIndices[Index];
-		FVulkanShaderHeader::FPackedUBInfo& PackedUB = OutHeader.PackedUBs[OutHeader.PackedUBs.AddZeroed()];
-		PackedUB.OriginalBindingIndex = PackedArrayInfo.VulkanBindingIndex;
-		PackedUB.PackedTypeIndex = CrossCompiler::PackedTypeNameToTypeIndex(PackedArrayInfo.TypeName);
-		PackedUB.SizeInBytes = OLDHeader.NEWPackedGlobalUBSizes[Index];
-
-		const FVulkanSpirv::FEntry* Entry = Spirv.GetEntryByBindingIndex(PackedArrayInfo.VulkanBindingIndex);
-		check(Entry);
-
-		// We are dealing with "HLSLCC_CB" for HLSLcc, and "$Globals" for DXC
-		check(Entry->Name.StartsWith(TEXT("HLSLCC_CB")) || Entry->Name.StartsWith(UBOGlobalsNameSpv));
-
-		PackedUB.SPIRVDescriptorSetOffset = Entry->WordDescriptorSetIndex;
-		PackedUB.SPIRVBindingIndexOffset = Entry->WordBindingIndex;
-	}
-
 	// Finally check for subpass/input attachments
 	if (BindingTable.InputAttachmentsMask != 0)
 	{
@@ -975,17 +742,11 @@ static void ConvertToNEWHeader(FOLDVulkanCodeHeader& OLDHeader,
 				continue;
 			}
 
-			const FString& AttachmentName = VULKAN_SUBPASS_FETCH_VAR_W[Index];
-			const FVulkanBindingTable::FBinding* Found = BindingTable.GetBindings().FindByPredicate([&AttachmentName](const FVulkanBindingTable::FBinding& Entry)
-				{
-					return Entry.Name == AttachmentName;
-				});
-			check(Found);
-			int32 BindingIndex = (int32)(Found - BindingTable.GetBindings().GetData());
-			check(BindingIndex >= 0 && BindingIndex <= BindingTable.GetBindings().Num());
-			FVulkanShaderHeader::EAttachmentType AttachmentType = AttachmentTypes[Index];
+			const FString AttachmentName(VULKAN_SUBPASS_FETCH_VAR_W[Index]);
+			const VkDescriptorType DescriptorType = VK_DESCRIPTOR_TYPE_INPUT_ATTACHMENT;
+			const FVulkanShaderHeader::EAttachmentType AttachmentType = AttachmentTypes[Index];
 			{
-				int32 HeaderGlobalIndex = GlobalNames.Find(AttachmentName);
+				const int32 HeaderGlobalIndex = GlobalNames.Find(AttachmentName);
 				check(HeaderGlobalIndex != INDEX_NONE);
 				check(GlobalNames[HeaderGlobalIndex] == AttachmentName);
 				FVulkanShaderHeader::FGlobalInfo& GlobalInfo = OutHeader.Globals[HeaderGlobalIndex];
@@ -993,34 +754,20 @@ static void ConvertToNEWHeader(FOLDVulkanCodeHeader& OLDHeader,
 				check(Entry);
 				check(Entry->Binding != -1);
 
-				VkDescriptorType DescriptorType = VK_DESCRIPTOR_TYPE_INPUT_ATTACHMENT;
 				GlobalInfo.OriginalBindingIndex = Entry->Binding;
 				OutHeader.GlobalSpirvInfos[HeaderGlobalIndex] = FVulkanShaderHeader::FSpirvInfo(Entry->WordDescriptorSetIndex, Entry->WordBindingIndex);
-				int32 GlobalDescriptorTypeIndex = OutHeader.GlobalDescriptorTypes.Add(DescriptorTypeToBinding(DescriptorType));
+				const int32 GlobalDescriptorTypeIndex = OutHeader.GlobalDescriptorTypes.Add(DescriptorTypeToBinding(DescriptorType));
 				GlobalInfo.TypeIndex = GlobalDescriptorTypeIndex;
 				GlobalInfo.CombinedSamplerStateAliasIndex = UINT16_MAX;
 #if VULKAN_ENABLE_BINDING_DEBUG_NAMES
 				GlobalInfo.DebugName = AttachmentName;
 #endif
-				int32 HeaderAttachmentIndex = OutHeader.InputAttachments.AddZeroed();
-				FVulkanShaderHeader::FInputAttachment& AttachmentInfo = OutHeader.InputAttachments[HeaderAttachmentIndex];
+				FVulkanShaderHeader::FInputAttachment& AttachmentInfo = OutHeader.InputAttachments.AddZeroed_GetRef();
 				AttachmentInfo.GlobalIndex = HeaderGlobalIndex;
 				AttachmentInfo.Type = AttachmentType;
 			}
 		}
 	}
-
-	check(OLDHeader.UniformBuffersCopyInfo.Num() == 0);
-	OutHeader.EmulatedUBsCopyInfo = OLDHeader.UniformBuffersCopyInfo;
-	OutHeader.EmulatedUBCopyRanges = OLDHeader.NEWEmulatedUBCopyRanges;
-	OutHeader.SourceHash = OLDHeader.SourceHash;
-	OutHeader.SpirvCRC = Spirv.CRC;
-#if VULKAN_ENABLE_BINDING_DEBUG_NAMES
-	OutHeader.DebugName = OLDHeader.ShaderName;
-#endif
-	OutHeader.InOutMask = OLDHeader.SerializedBindings.InOutMask.Bitmask;
-	OutHeader.RayTracingPayloadType = ShaderInput.Environment.GetCompileArgument(TEXT("RT_PAYLOAD_TYPE"), 0u);
-	OutHeader.RayTracingPayloadSize = ShaderInput.Environment.GetCompileArgument(TEXT("RT_PAYLOAD_MAX_SIZE"), 0u);
 }
 
 
@@ -1087,381 +834,120 @@ static FShaderResourceTable BuildSRTFromHeader(const FVulkanShaderHeader& NEWHea
 
 static void BuildShaderOutput(
 	FShaderCompilerOutput&		ShaderOutput,
-	const FShaderCompilerInput& ShaderInput,
-	const ANSICHAR*				InShaderSource,
-	int32						SourceLen,
+	const FVulkanShaderCompilerInternalState& InternalState,
+	const FSpirvReflectBindings& SpirvReflectBindings,
+	TMap<FString, FVulkanShaderHeader::EType>& EntryTypes,
 	const FVulkanBindingTable&	BindingTable,
-	uint32						NumLines,
 	uint8						WaveSize,
-	FVulkanSpirv&				Spirv,
+	const FVulkanSpirv&			Spirv,
 	const FString&				DebugName,
-	bool						bSourceContainsMetaDataOnly)
+	uint32						PackedGlobalArraySize,
+	TBitArray<>&				UsedUniformBufferSlots
+)
 {
-	const ANSICHAR* USFSource = InShaderSource;
-	FVulkanHlslccHeader CCHeader;
-	if (!CCHeader.Read(USFSource, SourceLen))
-	{
-		UE_LOG(LogVulkanShaderCompiler, Error, TEXT("Bad hlslcc header found: %s"), *ShaderInput.GenerateShaderName());
-		return;
-	}
+	const FShaderCompilerInput& ShaderInput = InternalState.Input;
+	const EShaderFrequency Frequency = InternalState.GetShaderFrequency();
 
-	if (!bSourceContainsMetaDataOnly && *USFSource != '#')
-	{
-		UE_LOG(LogVulkanShaderCompiler, Error, TEXT("Bad hlslcc header found with missing '#' character: %s"), *ShaderInput.GenerateShaderName());
-		return;
-	}
-
-	FOLDVulkanCodeHeader OLDHeader;
-
-	const EShaderFrequency Frequency = static_cast<EShaderFrequency>(ShaderOutput.Target.Frequency);
-
-	TBitArray<> UsedUniformBufferSlots;
-	const int32 MaxNumBits = VulkanBindless::MaxUniformBuffersPerStage * SF_NumFrequencies;
-	UsedUniformBufferSlots.Init(false, MaxNumBits);
-
-
-	static const FString AttributePrefix = TEXT("in_ATTRIBUTE");
-	static const FString GL_Prefix = TEXT("gl_");
-	for (auto& Input : CCHeader.Inputs)
-	{
-		// Only process attributes for vertex shaders.
-		if (Frequency == SF_Vertex && Input.Name.StartsWith(AttributePrefix))
-		{
-			int32 AttributeIndex = ParseNumber(*Input.Name + AttributePrefix.Len(), /*bEmptyIsZero:*/ true);
-			int32 Count = FMath::Max(1, Input.ArrayCount);
-			for(int32 Index = 0; Index < Count; ++Index)
-			{
-				OLDHeader.SerializedBindings.InOutMask.EnableField(Index + AttributeIndex);
-			}
-		}
-	}
+	FVulkanShaderHeader NEWHeader(FVulkanShaderHeader::EZero);
+	FShaderCompilerResourceTable ShaderResourceTable;
 
-	static const FString TargetPrefix = "out_Target";
-	static const FString GL_FragDepth = "gl_FragDepth";
-	for (auto& Output : CCHeader.Outputs)
-	{
-		// Only targets for pixel shaders must be tracked.
-		if (Frequency == SF_Pixel && Output.Name.StartsWith(TargetPrefix))
-		{
-			uint8 TargetIndex = ParseNumber(*Output.Name + TargetPrefix.Len(), /*bEmptyIsZero:*/ true);
-			OLDHeader.SerializedBindings.InOutMask.EnableField(TargetIndex);
-		}
-		// Only depth writes for pixel shaders must be tracked.
-		else if (Frequency == SF_Pixel && Output.Name.Equals(GL_FragDepth))
-		{
-			OLDHeader.SerializedBindings.InOutMask.EnableField(CrossCompiler::FShaderBindingInOutMask::DepthStencilMaskIndex);
-		}
-	}
+	NEWHeader.SpirvCRC = Spirv.CRC;
+	NEWHeader.RayTracingPayloadType = ShaderInput.Environment.GetCompileArgument(TEXT("RT_PAYLOAD_TYPE"), 0u);
+	NEWHeader.RayTracingPayloadSize = ShaderInput.Environment.GetCompileArgument(TEXT("RT_PAYLOAD_MAX_SIZE"), 0u);
+#if VULKAN_ENABLE_BINDING_DEBUG_NAMES
+	NEWHeader.DebugName = DebugName;
+#endif
 
+	// :todo-jn: Hash entire SPIRV for now, could eventually be removed since we use ShaderKeys
+	FSHA1::HashBuffer(Spirv.Data.GetData(), Spirv.Data.Num(), (uint8*)&NEWHeader.SourceHash);
 
-	TMap<FString, FVulkanShaderHeader::EType> NEWEntryTypes;
 
-	// Then 'normal' uniform buffers.
-	for (auto& UniformBlock : CCHeader.UniformBlocks)
+	// Flattens the array dimensions of the interface variable (aka shader attribute), e.g. from float4[2][3] -> float4[6]
+	auto FlattenAttributeArrayDimension = [](const SpvReflectInterfaceVariable& Attribute, uint32 FirstArrayDim = 0)
 	{
-		// DXC's generated "$Globals" has been converted to "_Globals" at this point
-		if (UniformBlock.Name.StartsWith(TEXT("HLSLCC_CB")) || UniformBlock.Name.StartsWith(TEXT("_Globals")))
+		uint32 FlattenedArrayDim = 1;
+		for (uint32 ArrayDimIndex = FirstArrayDim; ArrayDimIndex < Attribute.array.dims_count; ++ArrayDimIndex)
 		{
-			// Skip...
-		}
-		else
-		{
-			// Regular UB
-			const int32 VulkanBindingIndex = Spirv.FindBinding(UniformBlock.Name, true);
-			check(VulkanBindingIndex != -1);
-			check(!UsedUniformBufferSlots[VulkanBindingIndex]);
-			UsedUniformBufferSlots[VulkanBindingIndex] = true;
-
-			HandleReflectedUniformBuffer(UniformBlock.Name, VulkanBindingIndex, ShaderOutput);
-
-			++OLDHeader.SerializedBindings.NumUniformBuffers;
-			NEWEntryTypes.Add(*UniformBlock.Name, FVulkanShaderHeader::UniformBuffer);
+			FlattenedArrayDim *= Attribute.array.dims[ArrayDimIndex];
 		}
-	}
-
+		return FlattenedArrayDim;
+	};
 
-	const bool bSupportsBindless = ShaderInput.Environment.CompilerFlags.Contains(CFLAG_BindlessResources) || ShaderInput.Environment.CompilerFlags.Contains(CFLAG_BindlessSamplers);
-	const int32 StageOffset = bSupportsBindless ? (ShaderStage::GetStageForFrequency(Frequency) * VulkanBindless::MaxUniformBuffersPerStage) : 0;
 
-	const TArray<FVulkanBindingTable::FBinding>& HlslccBindings = BindingTable.GetBindings();
-	OLDHeader.NEWDescriptorInfo.NumBufferInfos = 0;
-	OLDHeader.NEWDescriptorInfo.NumImageInfos = 0;
-	for (int32 Index = 0; Index < HlslccBindings.Num(); ++Index)
+	// Only process input attributes for vertex shaders.
+	if (Frequency == SF_Vertex)
 	{
-		const FVulkanBindingTable::FBinding& Binding = HlslccBindings[Index];
-
-		OLDHeader.NEWDescriptorInfo.DescriptorTypes.Add(BindingToDescriptorType(Binding.Type));
-
-		switch (Binding.Type)
-		{
-		case EVulkanBindingType::Sampler:
-		case EVulkanBindingType::CombinedImageSampler:
-		case EVulkanBindingType::Image:
-		case EVulkanBindingType::StorageImage:
-		case EVulkanBindingType::InputAttachment:
-			++OLDHeader.NEWDescriptorInfo.NumImageInfos;
-			break;
-		case EVulkanBindingType::UniformBuffer:
-		case EVulkanBindingType::StorageBuffer:
-			++OLDHeader.NEWDescriptorInfo.NumBufferInfos;
-			break;
-		case EVulkanBindingType::PackedUniformBuffer:
-			{
-				FOLDVulkanCodeHeader::FPackedUBToVulkanBindingIndex& New = OLDHeader.NEWPackedUBToVulkanBindingIndices.AddDefaulted_GetRef();
-				New.TypeName = (CrossCompiler::EPackedTypeName)Binding.SubType;
-				New.VulkanBindingIndex = StageOffset + Index;
-				++OLDHeader.NEWDescriptorInfo.NumBufferInfos;
-			}
-			break;
-		case EVulkanBindingType::UniformTexelBuffer:
-		case EVulkanBindingType::StorageTexelBuffer:
-		case EVulkanBindingType::AccelerationStructure:
-			break;
-		default:
-			checkf(0, TEXT("Binding Type %d not found"), (int32)Binding.Type);
-			break;
-		}
-	}
-
-	const uint16 BytesPerComponent = 4;
+		static const FString AttributePrefix = TEXT("ATTRIBUTE");
 
-	// Packed global uniforms
-	TMap<CrossCompiler::EPackedTypeName, uint32> PackedGlobalArraySize;
-	for (auto& PackedGlobal : CCHeader.PackedGlobals)
-	{
-		int32 Found = -1;
-		for (int32 Index = 0; Index < OLDHeader.NEWPackedUBToVulkanBindingIndices.Num(); ++Index)
+		for (const SpvReflectInterfaceVariable* Attribute : SpirvReflectBindings.InputAttributes)
 		{
-			if (OLDHeader.NEWPackedUBToVulkanBindingIndices[Index].TypeName == (CrossCompiler::EPackedTypeName)PackedGlobal.PackedType)
+			if (CrossCompiler::FShaderConductorContext::IsIntermediateSpirvOutputVariable(Attribute->name))
 			{
-				Found = Index;
-				break;
+				continue;
 			}
-		}
-		check(Found != -1);
-
-		HandleReflectedGlobalConstantBufferMember(
-			PackedGlobal.Name,
-			Found,
-			PackedGlobal.Offset * BytesPerComponent,
-			PackedGlobal.Count * BytesPerComponent,
-			ShaderOutput
-		);
-
-		FString ParamName = PackedGlobal.Name;
-		FShaderParameterParser::RemoveBindlessParameterPrefix(ParamName);
-		NEWEntryTypes.Add(ParamName, FVulkanShaderHeader::PackedGlobal);
-
-		uint32& Size = PackedGlobalArraySize.FindOrAdd((CrossCompiler::EPackedTypeName)PackedGlobal.PackedType);
-		Size = FMath::Max<uint32>(BytesPerComponent * (PackedGlobal.Offset + PackedGlobal.Count), Size);
-	}
 
-	// Packed Uniform Buffers
-	TMap<int, TMap<CrossCompiler::EPackedTypeName, uint16> > PackedUniformBuffersSize;
-	OLDHeader.UNUSED_NumNonGlobalUBs = 0;
-	for (auto& PackedUB : CCHeader.PackedUBs)
-	{
-		//check(PackedUB.Attribute.Index == Header.SerializedBindings.NumUniformBuffers);
-		check(!UsedUniformBufferSlots[OLDHeader.UNUSED_NumNonGlobalUBs]);
-		UsedUniformBufferSlots[OLDHeader.UNUSED_NumNonGlobalUBs] = true;
-
-		HandleReflectedUniformBuffer(PackedUB.Attribute.Name, OLDHeader.UNUSED_NumNonGlobalUBs++, PackedUB.Attribute.Index, 0, ShaderOutput);
-
-		NEWEntryTypes.Add(PackedUB.Attribute.Name, FVulkanShaderHeader::PackedGlobal);
-	}
-
-	//#todo-rco: When using regular UBs, also set UsedUniformBufferSlots[] = 1
-
-	// Remap the destination UB index into the packed global array index
-	auto RemapDestIndexIntoPackedUB = [&OLDHeader](int8 DestUBTypeName)
-	{
-		for (int32 Index = 0; Index < OLDHeader.NEWPackedUBToVulkanBindingIndices.Num(); ++Index)
-		{
-			if (OLDHeader.NEWPackedUBToVulkanBindingIndices[Index].TypeName == (CrossCompiler::EPackedTypeName)DestUBTypeName)
+			if (!Attribute->semantic)
 			{
-				return Index;
+				continue;
 			}
-		}
-		check(0);
-		return -1;
-	};
-
-	for (auto& PackedUBCopy : CCHeader.PackedUBCopies)
-	{
-		// Not used: For flattening each UB into its own packed array (not a global one)
-		ensure(0);
-		CrossCompiler::FUniformBufferCopyInfo CopyInfo;
-		CopyInfo.SourceUBIndex = PackedUBCopy.SourceUB;
-		CopyInfo.SourceOffsetInFloats = PackedUBCopy.SourceOffset;
-		CopyInfo.DestUBTypeName = PackedUBCopy.DestPackedType;
-		CopyInfo.DestUBIndex = RemapDestIndexIntoPackedUB(CopyInfo.DestUBTypeName);
-		CopyInfo.DestUBTypeIndex = CrossCompiler::PackedTypeNameToTypeIndex(CopyInfo.DestUBTypeName);
-		CopyInfo.DestOffsetInFloats = PackedUBCopy.DestOffset;
-		CopyInfo.SizeInFloats = PackedUBCopy.Count;
-
-		OLDHeader.UniformBuffersCopyInfo.Add(CopyInfo);
-
-		auto& UniformBufferSize = PackedUniformBuffersSize.FindOrAdd(CopyInfo.DestUBIndex);
-		uint16& Size = UniformBufferSize.FindOrAdd((CrossCompiler::EPackedTypeName)CopyInfo.DestUBTypeName);
-		Size = FMath::Max<uint16>(BytesPerComponent * (CopyInfo.DestOffsetInFloats + CopyInfo.SizeInFloats), Size);
-	}
 
-	for (auto& PackedUBCopy : CCHeader.PackedUBGlobalCopies)
-	{
-		CrossCompiler::FUniformBufferCopyInfo CopyInfo;
-		CopyInfo.SourceUBIndex = PackedUBCopy.SourceUB;
-		CopyInfo.SourceOffsetInFloats = PackedUBCopy.SourceOffset;
-		CopyInfo.DestUBTypeName = PackedUBCopy.DestPackedType;
-		CopyInfo.DestUBIndex = RemapDestIndexIntoPackedUB(CopyInfo.DestUBTypeName);
-		CopyInfo.DestUBTypeIndex = CrossCompiler::PackedTypeNameToTypeIndex(CopyInfo.DestUBTypeName);
-		CopyInfo.DestOffsetInFloats = PackedUBCopy.DestOffset;
-		CopyInfo.SizeInFloats = PackedUBCopy.Count;
-
-		OLDHeader.UniformBuffersCopyInfo.Add(CopyInfo);
-
-		uint32& Size = PackedGlobalArraySize.FindOrAdd((CrossCompiler::EPackedTypeName)CopyInfo.DestUBTypeName);
-		Size = FMath::Max<uint32>(BytesPerComponent * (CopyInfo.DestOffsetInFloats + CopyInfo.SizeInFloats), Size);
-	}
-
-	// Generate a shortcut table for the PackedUBGlobalCopies
-	TMap<uint32, uint32> PackedUBGlobalCopiesRanges;
-	{
-		int32 MaxDestUBIndex = -1;
-		{
-			// Verify table is sorted
-			int32 PrevSourceUB = -1;
-			int32 Index = 0;
-			for (auto& Copy : OLDHeader.UniformBuffersCopyInfo)
+			const FString InputAttrName(ANSI_TO_TCHAR(Attribute->semantic));
+			if (InputAttrName.StartsWith(AttributePrefix))
 			{
-				if (PrevSourceUB < Copy.SourceUBIndex)
-				{
-					PrevSourceUB = Copy.SourceUBIndex;
-					MaxDestUBIndex = FMath::Max(MaxDestUBIndex, (int32)Copy.SourceUBIndex);
-					PackedUBGlobalCopiesRanges.Add(Copy.SourceUBIndex) = (Index << 16) | 1;
-				}
-				else if (PrevSourceUB == Copy.SourceUBIndex)
+				const uint32 AttributeIndex = ParseNumber(*InputAttrName + AttributePrefix.Len(), /*bEmptyIsZero:*/ true);
+				const uint32 FlattenedArrayDim = FlattenAttributeArrayDimension(*Attribute);
+				for (uint32 Index = 0; Index < FlattenedArrayDim; ++Index)
 				{
-					++PackedUBGlobalCopiesRanges.FindChecked(Copy.SourceUBIndex);
+					const uint32 BitIndex = (AttributeIndex + Index);
+					NEWHeader.InOutMask |= (1u << BitIndex);
 				}
-				else
-				{
-					// Internal error
-					check(0);
-				}
-				++Index;
-			}
-		}
-
-		OLDHeader.NEWEmulatedUBCopyRanges.AddZeroed(MaxDestUBIndex + 1);
-		for (int32 Index = 0; Index <= MaxDestUBIndex; ++Index)
-		{
-			uint32* Found = PackedUBGlobalCopiesRanges.Find(Index);
-			if (Found)
-			{
-				OLDHeader.NEWEmulatedUBCopyRanges[Index] = *Found;
 			}
 		}
 	}
 
-	// Update Packed global array sizes
-	OLDHeader.NEWPackedGlobalUBSizes.AddZeroed(OLDHeader.NEWPackedUBToVulkanBindingIndices.Num());
-	for (auto& Pair : PackedGlobalArraySize)
+	// Only process output attributes for pixel shaders.
+	if (Frequency == SF_Pixel)
 	{
-		CrossCompiler::EPackedTypeName TypeName = Pair.Key;
-		int32 PackedArrayIndex = -1;
-		for (int32 Index = 0; Index < OLDHeader.NEWPackedUBToVulkanBindingIndices.Num(); ++Index)
+		static const FString TargetPrefix = "SV_Target";
+
+		for (const SpvReflectInterfaceVariable* Attribute : SpirvReflectBindings.OutputAttributes)
 		{
-			if (OLDHeader.NEWPackedUBToVulkanBindingIndices[Index].TypeName == TypeName)
+			// Only depth writes for pixel shaders must be tracked.
+			if (Attribute->built_in == SpvBuiltInFragDepth)
 			{
-				PackedArrayIndex = Index;
-				break;
+				const uint32 BitIndex = (CrossCompiler::FShaderBindingInOutMask::DepthStencilMaskIndex);
+				NEWHeader.InOutMask |= (1u << BitIndex);
 			}
-		}
-		check(PackedArrayIndex != -1);
-		// In bytes
-		OLDHeader.NEWPackedGlobalUBSizes[PackedArrayIndex] = Align((uint32)Pair.Value, (uint32)16);
-	}
-
-	TSet<FString> SharedSamplerStates;
-	for (int32 i = 0; i < CCHeader.SamplerStates.Num(); i++)
-	{
-		const FString& Name = CCHeader.SamplerStates[i].Name;
-		int32 HlslccBindingIndex = Spirv.FindBinding(Name);
-		check(HlslccBindingIndex != -1);
-
-		SharedSamplerStates.Add(Name);
-		auto& Binding = HlslccBindings[HlslccBindingIndex];
-		int32 BindingIndex = Spirv.FindBinding(Binding.Name, true);
-		check(BindingIndex != -1);
-
-		HandleReflectedShaderSampler(Name, BindingIndex, ShaderOutput);
-
-		NEWEntryTypes.Add(Name, FVulkanShaderHeader::Global);
-
-		// Count only samplers states, not textures
-		OLDHeader.SerializedBindings.NumSamplers++;
-	}
-
-	for (auto& Sampler : CCHeader.Samplers)
-	{
-		int32 VulkanBindingIndex = Spirv.FindBinding(Sampler.Name, true);
-		check(VulkanBindingIndex != -1);
-
-		HandleReflectedShaderResource(Sampler.Name, Sampler.Offset, VulkanBindingIndex, Sampler.Count, ShaderOutput);
-
-		NEWEntryTypes.Add(Sampler.Name, FVulkanShaderHeader::Global);
-
-		for (auto& SamplerState : Sampler.SamplerStates)
-		{
-			if (!SharedSamplerStates.Contains(SamplerState))
+			else
 			{
-				// ParameterMap does not use a TMultiMap, so we cannot push the same entry to it more than once!  if we try to, we've done something wrong...
-				check(!ShaderOutput.ParameterMap.ContainsParameterAllocation(*SamplerState));
-
-				HandleReflectedShaderSampler(SamplerState, Sampler.Offset, VulkanBindingIndex, Sampler.Count, ShaderOutput);
-
-				NEWEntryTypes.Add(SamplerState, FVulkanShaderHeader::Global);
+				// Only targets for pixel shaders must be tracked.
+				const FString OutputAttrName(ANSI_TO_TCHAR(Attribute->semantic));
+				if (OutputAttrName.StartsWith(TargetPrefix))
+				{
+					const uint32 TargetIndex = ParseNumber(*OutputAttrName + TargetPrefix.Len(), /*bEmptyIsZero:*/ true);
 
-				// Count compiled texture-samplers as output samplers
-				OLDHeader.SerializedBindings.NumSamplers += Sampler.Count;
+					const uint32 FlattenedArrayDim = FlattenAttributeArrayDimension(*Attribute);
+					for (uint32 Index = 0; Index < FlattenedArrayDim; ++Index)
+					{
+						const uint32 BitIndex = (TargetIndex + Index);
+						NEWHeader.InOutMask |= (1u << BitIndex);
+					}
+				}
 			}
 		}
 	}
+	
+	const int32 StageOffset = InternalState.bSupportsBindless ? (ShaderStage::GetStageForFrequency(Frequency) * VulkanBindless::MaxUniformBuffersPerStage) : 0;
 
-	for (auto& UAV : CCHeader.UAVs)
-	{
-		int32 VulkanBindingIndex = Spirv.FindBinding(UAV.Name);
-		check(VulkanBindingIndex != -1);
-
-		HandleReflectedShaderUAV(UAV.Name, UAV.Offset, VulkanBindingIndex, UAV.Count, ShaderOutput);
-
-		NEWEntryTypes.Add(UAV.Name, FVulkanShaderHeader::Global);
-
-		OLDHeader.SerializedBindings.NumUAVs = FMath::Max<uint8>(
-			OLDHeader.SerializedBindings.NumUAVs,
-			UAV.Offset + UAV.Count
-			);
-	}
-
-	for (auto& AccelerationStructure : CCHeader.AccelerationStructures)
+	TArray<VkDescriptorType> DescriptorTypes;
+	const TArray<FVulkanBindingTable::FBinding>& HlslccBindings = BindingTable.GetBindings();
+	for (int32 Index = 0; Index < HlslccBindings.Num(); ++Index)
 	{
-		int32 VulkanBindingIndex = Spirv.FindBinding(AccelerationStructure.Name);
-		check(VulkanBindingIndex != -1);
-
-		HandleReflectedShaderResource(AccelerationStructure.Name, AccelerationStructure.Offset, VulkanBindingIndex, 1, ShaderOutput);
-
-		NEWEntryTypes.Add(AccelerationStructure.Name, FVulkanShaderHeader::Global);
-
-		OLDHeader.SerializedBindings.NumAccelerationStructures = FMath::Max<uint8>(
-			OLDHeader.SerializedBindings.NumAccelerationStructures,
-			AccelerationStructure.Offset + 1
-			);
+		const FVulkanBindingTable::FBinding& Binding = HlslccBindings[Index];
+		DescriptorTypes.Add(BindingToDescriptorType(Binding.Type));
 	}
 
-	// Lats make sure that there is some type of name visible
-	OLDHeader.ShaderName = CCHeader.Name.Len() > 0 ? CCHeader.Name : DebugName;
-
-	FSHA1::HashBuffer(USFSource, FCStringAnsi::Strlen(USFSource), (uint8*)&OLDHeader.SourceHash);
+	//#todo-rco: When using regular UBs, also set UsedUniformBufferSlots[] = 1
 
 	TArray<FString> OriginalParameters;
 	ShaderOutput.ParameterMap.GetAllParameterNames(OriginalParameters);
@@ -1477,27 +963,25 @@ static void BuildShaderOutput(
 		}
 
 		// Copy over the bits indicating which resource tables are active.
-		OLDHeader.SerializedBindings.ShaderResourceTable.ResourceTableBits = GenericSRT.ResourceTableBits;
-
-		OLDHeader.SerializedBindings.ShaderResourceTable.ResourceTableLayoutHashes = GenericSRT.ResourceTableLayoutHashes;
+		ShaderResourceTable.ResourceTableBits = GenericSRT.ResourceTableBits;
+		ShaderResourceTable.ResourceTableLayoutHashes = GenericSRT.ResourceTableLayoutHashes;
 
 		// Now build our token streams.
-		BuildResourceTableTokenStream(GenericSRT.TextureMap, GenericSRT.MaxBoundResourceTable, OLDHeader.SerializedBindings.ShaderResourceTable.TextureMap, true);
-		BuildResourceTableTokenStream(GenericSRT.ShaderResourceViewMap, GenericSRT.MaxBoundResourceTable, OLDHeader.SerializedBindings.ShaderResourceTable.ShaderResourceViewMap, true);
-		BuildResourceTableTokenStream(GenericSRT.SamplerMap, GenericSRT.MaxBoundResourceTable, OLDHeader.SerializedBindings.ShaderResourceTable.SamplerMap, true);
-		BuildResourceTableTokenStream(GenericSRT.UnorderedAccessViewMap, GenericSRT.MaxBoundResourceTable, OLDHeader.SerializedBindings.ShaderResourceTable.UnorderedAccessViewMap, true);
+		BuildResourceTableTokenStream(GenericSRT.TextureMap, GenericSRT.MaxBoundResourceTable, ShaderResourceTable.TextureMap, true);
+		BuildResourceTableTokenStream(GenericSRT.ShaderResourceViewMap, GenericSRT.MaxBoundResourceTable, ShaderResourceTable.ShaderResourceViewMap, true);
+		BuildResourceTableTokenStream(GenericSRT.SamplerMap, GenericSRT.MaxBoundResourceTable, ShaderResourceTable.SamplerMap, true);
+		BuildResourceTableTokenStream(GenericSRT.UnorderedAccessViewMap, GenericSRT.MaxBoundResourceTable, ShaderResourceTable.UnorderedAccessViewMap, true);
 	}
 
 	TArray<FString> NewParameters;
 	ShaderOutput.ParameterMap.GetAllParameterNames(NewParameters);
 
 	// Mark all used uniform buffer indices; however some are empty (eg GBuffers) so gather those as NewParameters
-	OLDHeader.UniformBuffersWithDescriptorMask = *UsedUniformBufferSlots.GetData();
 	uint16 NumParams = 0;
 	for (int32 Index = NewParameters.Num() - 1; Index >= 0; --Index)
 	{
 		uint16 OutIndex, OutBase, OutSize;
-		bool bFound = ShaderOutput.ParameterMap.FindParameterAllocation(*NewParameters[Index], OutIndex, OutBase, OutSize);
+		const bool bFound = ShaderOutput.ParameterMap.FindParameterAllocation(*NewParameters[Index], OutIndex, OutBase, OutSize);
 		ensure(bFound);
 		NumParams = FMath::Max((uint16)(OutIndex + 1), NumParams);
 		if (OriginalParameters.Contains(NewParameters[Index]))
@@ -1506,16 +990,21 @@ static void BuildShaderOutput(
 		}
 	}
 
-	// All newly added parameters are empty uniform buffers (with no constant data used), so no Vulkan Binding is required: remove from the mask
-	for (int32 Index = 0; Index < NewParameters.Num(); ++Index)
+	if (PackedGlobalArraySize > 0)
 	{
-		uint16 OutIndex, OutBase, OutSize;
-		ShaderOutput.ParameterMap.FindParameterAllocation(*NewParameters[Index], OutIndex, OutBase, OutSize);
-		OLDHeader.UniformBuffersWithDescriptorMask = OLDHeader.UniformBuffersWithDescriptorMask & ~((uint64)1 << (uint64)OutIndex);
+		FVulkanShaderHeader::FPackedUBInfo& PackedUB = NEWHeader.PackedUBs.AddZeroed_GetRef();
+		PackedUB.OriginalBindingIndex = StageOffset;
+		PackedUB.PackedTypeIndex = CrossCompiler::EPackedTypeIndex::HighP;
+		PackedUB.SizeInBytes = Align(PackedGlobalArraySize, 16u);
+
+		const FVulkanSpirv::FEntry* Entry = Spirv.GetEntryByBindingIndex(StageOffset);
+		check(Entry);
+		PackedUB.SPIRVDescriptorSetOffset = Entry->WordDescriptorSetIndex;
+		PackedUB.SPIRVBindingIndexOffset = Entry->WordBindingIndex;
 	}
 
-	FVulkanShaderHeader NEWHeader(FVulkanShaderHeader::EZero);
-	ConvertToNEWHeader(OLDHeader, BindingTable, Spirv, NEWEntryTypes, ShaderInput, CCHeader, ShaderOutput.ParameterMap, NEWHeader);
+	ConvertToHeader(ShaderResourceTable, BindingTable, DescriptorTypes, Spirv, EntryTypes, ShaderInput, SpirvReflectBindings, ShaderOutput.ParameterMap, NEWHeader);
+	check(NEWHeader.EmulatedUBsCopyInfo.Num() == 0);
 
 	if (ShaderInput.Environment.CompilerFlags.Contains(CFLAG_ExtraShaderData))
 	{
@@ -1538,9 +1027,6 @@ static void BuildShaderOutput(
 	Ar << SpirvCodeSizeBytes;
 	Ar.Serialize((uint8*)Spirv.Data.GetData(), SpirvCodeSizeBytes);
 
-	// Something to compare.
-	ShaderOutput.NumInstructions = NumLines;
-	ShaderOutput.NumTextureSamplers = OLDHeader.SerializedBindings.NumSamplers;
 	ShaderOutput.bSucceeded = true;
 
 	if (ShaderInput.ExtraSettings.bExtractShaderSource)
@@ -1701,7 +1187,7 @@ static uint32 CalculateSpirvInstructionCount(FVulkanSpirv& Spirv)
 static bool BuildShaderOutputFromSpirv(
 	CrossCompiler::FShaderConductorContext&	CompilerContext,
 	FVulkanSpirv&							Spirv,
-	const VulkanShaderCompilerInternalState InternalState,
+	const FVulkanShaderCompilerInternalState& InternalState,
 	FShaderCompilerOutput&					Output,
 	FVulkanBindingTable&					BindingTable,
 	uint8									WaveSize
@@ -1749,51 +1235,32 @@ static bool BuildShaderOutputFromSpirv(
 	// Build binding table
 	TMap<const SpvReflectDescriptorBinding*, int32> BindingToIndexMap;
 
-	/*for (const SpvReflectInterfaceVariable* Attribute : Bindings.InputAttributes)
-	{
-		check(Attribute->semantic != nullptr);
-		BindingTable.RegisterBinding(Attribute->semantic, "a", EVulkanBindingType::InputAttachment);
-	}*/
-
 	const FString UBOGlobalsNameSpv = ANSI_TO_TCHAR(CrossCompiler::FShaderConductorContext::GetIdentifierTable().GlobalsUniformBuffer);
 
-	for (const SpvReflectDescriptorBinding* Binding : Bindings.UniformBuffers)
-	{
-		FString ResourceName(ANSI_TO_TCHAR(Binding->name));
-		if (ResourceName == UBOGlobalsNameSpv)
-		{
-			int32 BindingIndex = BindingTable.RegisterBinding(TCHAR_TO_ANSI(*UBOGlobalsNameSpv), "h", EVulkanBindingType::PackedUniformBuffer);
-			BindingToIndexMap.Add(Binding, BindingIndex);
-			break;
-		}
-	}
-
-	for (const SpvReflectDescriptorBinding* Binding : Bindings.UniformBuffers)
-	{
-		const FString ResourceName(ANSI_TO_TCHAR(Binding->name));
-		if (ResourceName != UBOGlobalsNameSpv)
-		{
-			int32 BindingIndex = BindingTable.RegisterBinding(Binding->name, "u", EVulkanBindingType::UniformBuffer);
-			BindingToIndexMap.Add(Binding, BindingIndex);
-		}
-	}
-
-	for (const SpvReflectDescriptorBinding* Binding : Bindings.InputAttachments)
-	{
-		int32 BindingIndex = BindingTable.RegisterBinding(Binding->name, "a", EVulkanBindingType::InputAttachment);
-		BindingToIndexMap.Add(Binding, BindingIndex);
-		BindingTable.InputAttachmentsMask |= (1u << Binding->input_attachment_index);
-	}
-
-	auto RegisterBindings = [&BindingTable, &BindingToIndexMap](TArray<SpvReflectDescriptorBinding*>& Bindings, const char* BlockName, EVulkanBindingType::EType BindingType)
+	auto RegisterBindings = [&BindingTable, &BindingToIndexMap, &UBOGlobalsNameSpv](TArray<SpvReflectDescriptorBinding*>& Bindings, const char* BlockName, EVulkanBindingType::EType BindingType)
 	{
 		for (const SpvReflectDescriptorBinding* Binding : Bindings)
 		{
+			if (((BindingType == EVulkanBindingType::PackedUniformBuffer) && (UBOGlobalsNameSpv != Binding->name)) ||
+				((BindingType == EVulkanBindingType::UniformBuffer) && (UBOGlobalsNameSpv == Binding->name)))
+			{
+				continue;
+			}
+
 			const int32 BindingIndex = BindingTable.RegisterBinding(Binding->name, BlockName, BindingType);
 			BindingToIndexMap.Add(Binding, BindingIndex);
+
+			if (BindingType == EVulkanBindingType::InputAttachment)
+			{
+				BindingTable.InputAttachmentsMask |= (1u << Binding->input_attachment_index);
+			}
 		}
 	};
 
+	RegisterBindings(Bindings.UniformBuffers, "h", EVulkanBindingType::PackedUniformBuffer);
+	RegisterBindings(Bindings.UniformBuffers, "u", EVulkanBindingType::UniformBuffer);
+	RegisterBindings(Bindings.InputAttachments, "a", EVulkanBindingType::InputAttachment);
+
 	RegisterBindings(Bindings.TBufferUAVs, "u", EVulkanBindingType::StorageTexelBuffer);
 	RegisterBindings(Bindings.SBufferUAVs, "u", EVulkanBindingType::StorageBuffer);
 	RegisterBindings(Bindings.TextureUAVs, "u", EVulkanBindingType::StorageImage);
@@ -1808,161 +1275,146 @@ static bool BuildShaderOutputFromSpirv(
 	// Sort binding table
 	BindingTable.SortBindings();
 
-	CrossCompiler::FHlslccHeaderWriter CCHeaderWriter;
-
-	// Iterate over all resource bindings grouped by resource type
-	for (const SpvReflectInterfaceVariable* Attribute : Bindings.InputAttributes)
-	{
-		CCHeaderWriter.WriteInputAttribute(*Attribute);
-	}
+	uint32 PackedGlobalArraySize = 0;
+	TBitArray<> UsedUniformBufferSlots;
+	const int32 MaxNumBits = VulkanBindless::MaxUniformBuffersPerStage * SF_NumFrequencies;
+	UsedUniformBufferSlots.Init(false, MaxNumBits);
 
-	for (const SpvReflectInterfaceVariable* Attribute : Bindings.OutputAttributes)
-	{
-		CCHeaderWriter.WriteOutputAttribute(*Attribute);
-	}
+	TMap<FString, FVulkanShaderHeader::EType> EntryTypes;
 
-	// Final descriptor binding numbers for UBs
+	// Final descriptor binding numbers for all other resource types
 	{
 		const int32 StageOffset = InternalState.bSupportsBindless ? (ShaderStage::GetStageForFrequency(ShaderFrequency) * VulkanBindless::MaxUniformBuffersPerStage) : 0;
 		const uint32_t DescSetNumber = InternalState.bSupportsBindless ? (uint32_t)VulkanBindless::BindlessSingleUseUniformBufferSet : (uint32_t)SPV_REFLECT_SET_NUMBER_DONT_CHANGE;
 
-		int32 UBOBindings = 0;
-		auto AddUBReflectionInfo = [&](const SpvReflectDescriptorBinding* Binding, const FString& BindingName, const TCHAR* ResourceName)
-		{
-			// Register binding for uniform buffer
-			const int32 BindingIndex = BindingTable.GetRealBindingIndex(BindingToIndexMap[Binding]) + StageOffset;
-
-			SpvReflectResult SpvResult = Reflection.ChangeDescriptorBindingNumbers(Binding, BindingIndex, DescSetNumber);
-			check(SpvResult == SPV_REFLECT_RESULT_SUCCESS);
-
-			Spirv.ReflectionInfo.Add(FVulkanSpirv::FEntry(BindingName, BindingIndex));
-
-			CCHeaderWriter.WriteUniformBlock(ResourceName, UBOBindings++);
-		};
-
-		// Process Globals first
-		for (const SpvReflectDescriptorBinding* Binding : Bindings.UniformBuffers)
+		auto AddReflectionInfos = [&](TArray<SpvReflectDescriptorBinding*>& BindingArray, EVulkanBindingType::EType BindingType, int32 BindingOffset)
 		{
-			const FString ResourceName(ANSI_TO_TCHAR(Binding->name));
-
-			if (ResourceName == UBOGlobalsNameSpv)
+			for (const SpvReflectDescriptorBinding* Binding : BindingArray)
 			{
-				AddUBReflectionInfo(Binding, UBOGlobalsNameSpv, TEXT("_Globals_h"));
+				checkf(!InternalState.bSupportsBindless || (BindingType == EVulkanBindingType::UniformBuffer), TEXT("Bindless shaders should only have uniform buffers."));
 
-				// Register all uniform buffer members as loose data
-				for (uint32 MemberIndex = 0; MemberIndex < Binding->block.member_count; ++MemberIndex)
+				const FString ResourceName(ANSI_TO_TCHAR(Binding->name));
+
+				if (((BindingType == EVulkanBindingType::PackedUniformBuffer) && (ResourceName != UBOGlobalsNameSpv)) ||
+					((BindingType == EVulkanBindingType::UniformBuffer) && (ResourceName == UBOGlobalsNameSpv)))
 				{
-					const SpvReflectBlockVariable* Member = &(Binding->block.members[MemberIndex]);
-					CCHeaderWriter.WritePackedGlobal(ANSI_TO_TCHAR(Member->name), CrossCompiler::EPackedTypeName::HighP, Member->absolute_offset, Member->size);
+					continue;
 				}
 
-				// Stop after we found $Globals uniform buffer
-				break;
-			}
-		}
-
-		// Process all other UB
-		for (const SpvReflectDescriptorBinding* Binding : Bindings.UniformBuffers)
-		{
-			const FString ResourceName(ANSI_TO_TCHAR(Binding->name));
-
-			if (ResourceName != UBOGlobalsNameSpv)
-			{
-				AddUBReflectionInfo(Binding, ResourceName, *ResourceName);
-			}
-		}
-	}
+				const int32 BindingIndex = BindingTable.GetRealBindingIndex(BindingToIndexMap[Binding]) + StageOffset;
 
-	// Final descriptor binding numbers for all other resource types
-	{
-		int32 UAVBindings = 0, SRVBindings = 0, SMPBindings = 0, ASBindings = 0;
-		auto AddReflectionInfos = [&](TArray<SpvReflectDescriptorBinding*>& BindingArray, EVulkanBindingType::EType BindingType)
-		{
-			for (const SpvReflectDescriptorBinding* Binding : BindingArray)
-			{
-				const int32 BindingIndex = BindingTable.GetRealBindingIndex(BindingToIndexMap[Binding]);
-
-				SpvReflectResult SpvResult = Reflection.ChangeDescriptorBindingNumbers(Binding, BindingIndex);
+				const SpvReflectResult SpvResult = Reflection.ChangeDescriptorBindingNumbers(Binding, BindingIndex, DescSetNumber);
 				check(SpvResult == SPV_REFLECT_RESULT_SUCCESS);
 
-				const FString ResourceName(ANSI_TO_TCHAR(Binding->name));
+				const int32 VulkanBindingIndex = Spirv.ReflectionInfo.Add(FVulkanSpirv::FEntry(ResourceName, BindingIndex));
 
 				switch (BindingType)
 				{
 				case EVulkanBindingType::StorageTexelBuffer:
 				case EVulkanBindingType::StorageBuffer:
 				case EVulkanBindingType::StorageImage:
-					CCHeaderWriter.WriteUAV(*ResourceName, UAVBindings++);
-					break;
-
-				case EVulkanBindingType::UniformTexelBuffer:
-					CCHeaderWriter.WriteSRV(*ResourceName, SRVBindings++);
+					HandleReflectedShaderUAV(ResourceName, BindingOffset, VulkanBindingIndex, 1, Output);
+					EntryTypes.Add(ResourceName, FVulkanShaderHeader::Global);
 					break;
 
 				case EVulkanBindingType::Image:
-				{
-					if (Binding->usage_binding_count > 0)
-					{
-						TArray<FString> AssociatedResourceNames;
-						AssociatedResourceNames.SetNum(Binding->usage_binding_count);
+					// todo-jn: Could verify that we have the samplers...
+					//for (uint32 UsageIndex = 0; UsageIndex < Binding->usage_binding_count; ++UsageIndex)
+					//{
+					//	const SpvReflectDescriptorBinding* AssociatedResource = Binding->usage_bindings[UsageIndex];
+					//	AssociatedResourceNames[UsageIndex] = ANSI_TO_TCHAR(AssociatedResource->name);
+					//}
+					[[fallthrough]];
 
-						for (uint32 UsageIndex = 0; UsageIndex < Binding->usage_binding_count; ++UsageIndex)
-						{
-							const SpvReflectDescriptorBinding* AssociatedResource = Binding->usage_bindings[UsageIndex];
-							AssociatedResourceNames[UsageIndex] = ANSI_TO_TCHAR(AssociatedResource->name);
-						}
-
-						CCHeaderWriter.WriteSRV(*ResourceName, SRVBindings++, /*Count:*/ 1, AssociatedResourceNames);
-					}
-					else
-					{
-						CCHeaderWriter.WriteSRV(*ResourceName, SRVBindings++);
-					}
-				}
-				break;
+				case EVulkanBindingType::UniformTexelBuffer:
+					HandleReflectedShaderResource(ResourceName, BindingOffset, VulkanBindingIndex, 1, Output);
+					EntryTypes.Add(ResourceName, FVulkanShaderHeader::Global);
+					break;
 
 				case EVulkanBindingType::Sampler:
-					CCHeaderWriter.WriteSamplerState(*ResourceName, SMPBindings++);
+					HandleReflectedShaderSampler(ResourceName, VulkanBindingIndex, Output);
+					//HandleReflectedShaderSampler(ResourceName, BindingOffset, VulkanBindingIndex, 1, Output);
+					EntryTypes.Add(ResourceName, FVulkanShaderHeader::Global);
 					break;
 
 				case EVulkanBindingType::AccelerationStructure:
-					CCHeaderWriter.WriteAccelerationStructures(*ResourceName, ASBindings++);
+					HandleReflectedShaderResource(ResourceName, BindingOffset, VulkanBindingIndex, 1, Output);
+					EntryTypes.Add(ResourceName, FVulkanShaderHeader::Global);
 					break;
 
 				case EVulkanBindingType::InputAttachment:
 					// Do Nothing
 					break;
 
+				case EVulkanBindingType::PackedUniformBuffer:
+					{
+						// Use the given global ResourceName instead of patching it to _Globals_h
+						check(!UsedUniformBufferSlots[VulkanBindingIndex]);
+						UsedUniformBufferSlots[VulkanBindingIndex] = true;
+						//HandleReflectedUniformBuffer(ResourceName, VulkanBindingIndex, Output);
+						//EntryTypes.Add(ResourceName, FVulkanShaderHeader::UniformBuffer);
+
+						// Register all uniform buffer members of Globals as loose data
+						for (uint32 MemberIndex = 0; MemberIndex < Binding->block.member_count; ++MemberIndex)
+						{
+							const SpvReflectBlockVariable& Member = Binding->block.members[MemberIndex];
+
+							FString MemberName(ANSI_TO_TCHAR(Member.name));
+							HandleReflectedGlobalConstantBufferMember(
+								MemberName,
+								BindingOffset,
+								Member.absolute_offset,
+								Member.size,
+								Output
+							);
+
+							FShaderParameterParser::RemoveBindlessParameterPrefix(MemberName);
+							EntryTypes.Add(MemberName, FVulkanShaderHeader::PackedGlobal);
+
+							PackedGlobalArraySize = FMath::Max<uint32>((Member.absolute_offset + Member.size), PackedGlobalArraySize);
+						}
+					}
+					break;
+
+				case EVulkanBindingType::UniformBuffer:
+					{
+						check(!UsedUniformBufferSlots[VulkanBindingIndex]);
+						UsedUniformBufferSlots[VulkanBindingIndex] = true;
+						HandleReflectedUniformBuffer(ResourceName, VulkanBindingIndex, Output);
+						EntryTypes.Add(ResourceName, FVulkanShaderHeader::UniformBuffer);
+					}
+					break;
+
 				default:
 					check(false);
 					break;
 				};
 
-				Spirv.ReflectionInfo.Add(FVulkanSpirv::FEntry(ResourceName, BindingIndex));
+				BindingOffset++;
 			}
+			return BindingOffset;
 		};
 
+		// Process Globals first (PackedUniformBuffer) and then regular UBs
+		const int32 GlobalUBCount = AddReflectionInfos(Bindings.UniformBuffers, EVulkanBindingType::PackedUniformBuffer, 0);
+		int32 UBOBindings = AddReflectionInfos(Bindings.UniformBuffers, EVulkanBindingType::UniformBuffer, 0) + GlobalUBCount;
 
-		AddReflectionInfos(Bindings.InputAttachments, EVulkanBindingType::InputAttachment);
+		AddReflectionInfos(Bindings.InputAttachments, EVulkanBindingType::InputAttachment, 0);
 
-		AddReflectionInfos(Bindings.TBufferUAVs, EVulkanBindingType::StorageTexelBuffer);
-		AddReflectionInfos(Bindings.SBufferUAVs, EVulkanBindingType::StorageBuffer);
-		AddReflectionInfos(Bindings.TextureUAVs, EVulkanBindingType::StorageImage);
+		int32 UAVBindings = 0;
+		UAVBindings = AddReflectionInfos(Bindings.TBufferUAVs, EVulkanBindingType::StorageTexelBuffer, UAVBindings);
+		UAVBindings = AddReflectionInfos(Bindings.SBufferUAVs, EVulkanBindingType::StorageBuffer, UAVBindings);
+		UAVBindings = AddReflectionInfos(Bindings.TextureUAVs, EVulkanBindingType::StorageImage, UAVBindings);
 
-		AddReflectionInfos(Bindings.TBufferSRVs, EVulkanBindingType::UniformTexelBuffer);
-		AddReflectionInfos(Bindings.SBufferSRVs, EVulkanBindingType::UniformTexelBuffer);
-		AddReflectionInfos(Bindings.TextureSRVs, EVulkanBindingType::Image);
+		int32 SRVBindings = 0;
+		SRVBindings = AddReflectionInfos(Bindings.TBufferSRVs, EVulkanBindingType::UniformTexelBuffer, SRVBindings);
+		SRVBindings = AddReflectionInfos(Bindings.SBufferSRVs, EVulkanBindingType::UniformTexelBuffer, SRVBindings);
+		SRVBindings = AddReflectionInfos(Bindings.TextureSRVs, EVulkanBindingType::Image, SRVBindings);
 
-		AddReflectionInfos(Bindings.Samplers, EVulkanBindingType::Sampler);
-		AddReflectionInfos(Bindings.AccelerationStructures, EVulkanBindingType::AccelerationStructure);
+		Output.NumTextureSamplers = AddReflectionInfos(Bindings.Samplers, EVulkanBindingType::Sampler, 0);
+		AddReflectionInfos(Bindings.AccelerationStructures, EVulkanBindingType::AccelerationStructure, 0);
 	}
 
-	// Build final shader output with meta data
-	CCHeaderWriter.WriteSourceInfo(*InternalState.Input.VirtualSourceFilePath, *InternalState.GetEntryPointName());
-	CCHeaderWriter.WriteCompilerInfo();
-
-	const FString MetaData = CCHeaderWriter.ToString();
-
 	Output.Target = InternalState.Input.Target;
 
 	// Overwrite updated SPIRV code
@@ -1991,27 +1443,27 @@ static bool BuildShaderOutputFromSpirv(
 	}
 
 	PatchSpirvReflectionEntries(Spirv);
+
 	Spirv.EntryPointName = PatchSpirvEntryPointWithCRC(Spirv, Spirv.CRC);
 
-	const uint32 ApproxInstructionCount = CalculateSpirvInstructionCount(Spirv);
+	Output.NumInstructions = CalculateSpirvInstructionCount(Spirv);
 
 	BuildShaderOutput(
 		Output,
-		InternalState.Input,
-		TCHAR_TO_ANSI(*MetaData),
-		MetaData.Len(),
+		InternalState,
+		Bindings,
+		EntryTypes,
 		BindingTable,
-		ApproxInstructionCount,
 		WaveSize,
 		Spirv,
 		InternalState.GetDebugName(),
-		true // source contains meta data only
+		PackedGlobalArraySize,
+		UsedUniformBufferSlots
 	);
 
 	if (InternalState.bDebugDump)
 	{
 		// Write meta data to debug output file and write SPIR-V dump in binary and text form
-		DumpDebugShaderText(InternalState.Input, MetaData, TEXT("meta.txt"));
 		DumpDebugShaderBinary(InternalState.Input, Spirv.GetByteData(), Spirv.GetByteSize(), TEXT("spv"));
 		DumpDebugShaderDisassembledSpirv(InternalState.Input, Spirv.GetByteData(), Spirv.GetByteSize(), TEXT("spvasm"));
 	}
@@ -2146,7 +1598,7 @@ static void Patch64bitSamplers(FVulkanSpirv& Spirv)
 	}
 }
 
-static FString VulkanGetShaderProfileDXC(const VulkanShaderCompilerInternalState& InternalState, const CrossCompiler::FShaderConductorOptions Options)
+static FString VulkanGetShaderProfileDXC(const FVulkanShaderCompilerInternalState& InternalState, const CrossCompiler::FShaderConductorOptions Options)
 {
 	const TCHAR* ShaderProfile = TEXT("unknown");
 	switch (InternalState.GetShaderFrequency())
@@ -2178,7 +1630,7 @@ static FString VulkanGetShaderProfileDXC(const VulkanShaderCompilerInternalState
 }
 
 static void VulkanCreateDXCCompileBatchFiles(
-	const VulkanShaderCompilerInternalState& InternalState,
+	const FVulkanShaderCompilerInternalState& InternalState,
 	const CrossCompiler::FShaderConductorOptions Options)
 {
 	FString USFFilename = InternalState.Input.GetSourceFilename();
@@ -2347,7 +1799,7 @@ static FString ParseEntrypointDecl(FStringView PreprocessedShader, FStringView E
 }
 
 uint8 ParseWaveSize(
-	const VulkanShaderCompilerInternalState InternalState,
+	const FVulkanShaderCompilerInternalState& InternalState,
 	const FString& PreprocessedShader
 	)
 {
@@ -2399,7 +1851,7 @@ uint8 ParseWaveSize(
 }
 
 static bool CompileWithShaderConductor(
-	const VulkanShaderCompilerInternalState InternalState,
+	const FVulkanShaderCompilerInternalState& InternalState,
 	const FString&			PreprocessedShader,
 	FShaderCompilerOutput&	Output,
 	FVulkanBindingTable&	BindingTable
@@ -2501,7 +1953,7 @@ static bool CompileWithShaderConductor(
 
 bool PreprocessVulkanShader(const FShaderCompilerInput& Input, const FShaderCompilerEnvironment& Environment, FShaderPreprocessOutput& PreprocessOutput)
 {
-	const VulkanShaderCompilerInternalState InternalState(Input);
+	const FVulkanShaderCompilerInternalState InternalState(Input);
 
 	PRAGMA_DISABLE_DEPRECATION_WARNINGS		// FShaderCompilerDefinitions will be made internal in the future, marked deprecated until then
 
@@ -2608,7 +2060,7 @@ void CompileVulkanShader(const FShaderCompilerInput& Input, const FShaderPreproc
 {
 	check(IsVulkanShaderFormat(Input.ShaderFormat));
 
-	const VulkanShaderCompilerInternalState InternalState(Input);
+	const FVulkanShaderCompilerInternalState InternalState(Input);
 
 	const EHlslShaderFrequency HlslFrequency = InternalState.GetHlslShaderFrequency();
 	if (HlslFrequency == HSF_InvalidFrequency)
