VulkanRHI: (VK RT WIP) Move all serialized parameter of the shader compilation process into a struct together (Header, SRT, Spirv).

#rnx

diff --git a/Engine/Source/Developer/VulkanShaderFormat/Private/VulkanShaderCompiler.cpp b/Engine/Source/Developer/VulkanShaderFormat/Private/VulkanShaderCompiler.cpp
--- a/Engine/Source/Developer/VulkanShaderFormat/Private/VulkanShaderCompiler.cpp
+++ b/Engine/Source/Developer/VulkanShaderFormat/Private/VulkanShaderCompiler.cpp
@@ -304,13 +304,24 @@ struct FVulkanShaderCompilerInternalState
 	}
 };
 
+// Data structures that will get serialized into ShaderCompilerOutput
+struct VulkanShaderCompilerSerializedOutput
+{
+	VulkanShaderCompilerSerializedOutput()
+		: Header(FVulkanShaderHeader::EZero)
+	{
+	}
+
+	FVulkanShaderHeader Header;
+	FShaderResourceTable ShaderResourceTable;
+	FVulkanSpirv Spirv;
+};
 
 
 static int32 AddGlobal( const TArray<VkDescriptorType>& DescriptorTypes,
 						const FString& ParameterName,
 						uint16 BindingIndex,
-						const FVulkanSpirv& Spirv,
-						FVulkanShaderHeader& OutHeader,
+						VulkanShaderCompilerSerializedOutput& SerializedOutput, 
 						const TArray<FString>& GlobalNames
 )
 {
@@ -318,21 +329,21 @@ static int32 AddGlobal( const TArray<VkDescriptorType>& DescriptorTypes,
 	check(HeaderGlobalIndex != INDEX_NONE);
 	check(GlobalNames[HeaderGlobalIndex] == ParameterName);
 
-	FVulkanShaderHeader::FGlobalInfo& GlobalInfo = OutHeader.Globals[HeaderGlobalIndex];
-	const FVulkanSpirv::FEntry* Entry = Spirv.GetEntry(ParameterName);
+	FVulkanShaderHeader::FGlobalInfo& GlobalInfo = SerializedOutput.Header.Globals[HeaderGlobalIndex];
+	const FVulkanSpirv::FEntry* Entry = SerializedOutput.Spirv.GetEntry(ParameterName);
 	if (Entry)
 	{
 		if (Entry->Binding == -1)
 		{
 			// Texel buffers get put into a uniform block
-			Entry = Spirv.GetEntry(ParameterName + TEXT("_BUFFER"));
+			Entry = SerializedOutput.Spirv.GetEntry(ParameterName + TEXT("_BUFFER"));
 			check(Entry);
 			check(Entry->Binding != -1);
 		}
 	}
 	else
 	{
-		Entry = Spirv.GetEntryByBindingIndex(BindingIndex);
+		Entry = SerializedOutput.Spirv.GetEntryByBindingIndex(BindingIndex);
 		check(Entry);
 		check(Entry->Binding != -1);
 		if (!Entry->Name.EndsWith(TEXT("_BUFFER")))
@@ -344,9 +355,9 @@ static int32 AddGlobal( const TArray<VkDescriptorType>& DescriptorTypes,
 	const VkDescriptorType DescriptorType = DescriptorTypes[Entry->Binding];
 
 	GlobalInfo.OriginalBindingIndex = Entry->Binding;
-	OutHeader.GlobalSpirvInfos[HeaderGlobalIndex] = FVulkanShaderHeader::FSpirvInfo(Entry->WordDescriptorSetIndex, Entry->WordBindingIndex);
+	SerializedOutput.Header.GlobalSpirvInfos[HeaderGlobalIndex] = FVulkanShaderHeader::FSpirvInfo(Entry->WordDescriptorSetIndex, Entry->WordBindingIndex);
 
-	const int32 GlobalDescriptorTypeIndex = OutHeader.GlobalDescriptorTypes.Add(DescriptorTypeToBinding(DescriptorType));
+	const int32 GlobalDescriptorTypeIndex = SerializedOutput.Header.GlobalDescriptorTypes.Add(DescriptorTypeToBinding(DescriptorType));
 	GlobalInfo.TypeIndex = GlobalDescriptorTypeIndex;
 	GlobalInfo.CombinedSamplerStateAliasIndex = UINT16_MAX;
 
@@ -362,9 +373,8 @@ static void AddUBResources( const FString& UBName,
 							uint32 BufferIndex,
 							const TArray<uint32>& BindingArray,
 							const TArray<VkDescriptorType>& DescriptorTypes,
-							const FVulkanSpirv& Spirv,
 							FVulkanShaderHeader::FUniformBufferInfo& OutUBInfo,
-							FVulkanShaderHeader& OutHeader,
+							VulkanShaderCompilerSerializedOutput& SerializedOutput,
 							TArray<FString>& GlobalNames)
 {
 	if (BindingArray.Num() > 0)
@@ -391,7 +401,7 @@ static void AddUBResources( const FString& UBName,
 
 				FVulkanShaderHeader::FUBResourceInfo& UBResourceInfo = OutUBInfo.ResourceEntries.AddZeroed_GetRef();;
 
-				const int32 HeaderGlobalIndex = AddGlobal(DescriptorTypes, ResourceTableEntry.UniformBufferMemberName, BindingIndex, Spirv, OutHeader, GlobalNames);
+				const int32 HeaderGlobalIndex = AddGlobal(DescriptorTypes, ResourceTableEntry.UniformBufferMemberName, BindingIndex, SerializedOutput, GlobalNames);
 				UBResourceInfo.SourceUBResourceIndex = ResourceIndex;
 				UBResourceInfo.OriginalBindingIndex = BindingIndex;
 				UBResourceInfo.GlobalIndex = HeaderGlobalIndex;
@@ -411,14 +421,15 @@ static void AddUniformBuffer(
 	const FShaderCompilerResourceTable& ShaderResourceTable,
 	const TArray<VkDescriptorType>& DescriptorTypes,
 	const FShaderCompilerInput& ShaderInput,
-	const FVulkanSpirv& Spirv,
 	const FString& UBName,
 	uint16 BindingIndex,
 	FShaderParameterMap& InOutParameterMap,
-	FVulkanShaderHeader& OutHeader,
+	VulkanShaderCompilerSerializedOutput& SerializedOutput,
 	TArray<FString>& GlobalNames
 )
 {
+	FVulkanShaderHeader& OutHeader = SerializedOutput.Header;
+
 	const int32 HeaderUBIndex = OutHeader.UniformBuffers.AddZeroed();
 	FVulkanShaderHeader::FUniformBufferInfo& UBInfo = OutHeader.UniformBuffers[HeaderUBIndex];
 
@@ -440,7 +451,7 @@ static void AddUniformBuffer(
 	UBInfo.DebugName = UBName;
 #endif
 
-	const FVulkanSpirv::FEntry* Entry = Spirv.GetEntry(UBName);
+	const FVulkanSpirv::FEntry* Entry = SerializedOutput.Spirv.GetEntry(UBName);
 	if (Entry)
 	{
 		UBInfo.bOnlyHasResources = false;
@@ -462,10 +473,10 @@ static void AddUniformBuffer(
 	if (ShaderResourceTable.ResourceTableBits & (1 << BindingIndex))
 	{
 		// Make sure to process in the same order as when gathering names below
-		AddUBResources(UBName, ShaderInput.Environment.ResourceTableMap, BindingIndex, ShaderResourceTable.TextureMap, DescriptorTypes, Spirv, UBInfo, OutHeader, GlobalNames);
-		AddUBResources(UBName, ShaderInput.Environment.ResourceTableMap, BindingIndex, ShaderResourceTable.SamplerMap, DescriptorTypes, Spirv, UBInfo, OutHeader, GlobalNames);
-		AddUBResources(UBName, ShaderInput.Environment.ResourceTableMap, BindingIndex, ShaderResourceTable.ShaderResourceViewMap, DescriptorTypes, Spirv, UBInfo, OutHeader, GlobalNames);
-		AddUBResources(UBName, ShaderInput.Environment.ResourceTableMap, BindingIndex, ShaderResourceTable.UnorderedAccessViewMap, DescriptorTypes, Spirv, UBInfo, OutHeader, GlobalNames);
+		AddUBResources(UBName, ShaderInput.Environment.ResourceTableMap, BindingIndex, ShaderResourceTable.TextureMap, DescriptorTypes, UBInfo, SerializedOutput, GlobalNames);
+		AddUBResources(UBName, ShaderInput.Environment.ResourceTableMap, BindingIndex, ShaderResourceTable.SamplerMap, DescriptorTypes, UBInfo, SerializedOutput, GlobalNames);
+		AddUBResources(UBName, ShaderInput.Environment.ResourceTableMap, BindingIndex, ShaderResourceTable.ShaderResourceViewMap, DescriptorTypes, UBInfo, SerializedOutput, GlobalNames);
+		AddUBResources(UBName, ShaderInput.Environment.ResourceTableMap, BindingIndex, ShaderResourceTable.UnorderedAccessViewMap, DescriptorTypes, UBInfo, SerializedOutput, GlobalNames);
 	}
 	else
 	{
@@ -649,13 +660,15 @@ static void ConvertToHeader(
 	FShaderCompilerResourceTable& ShaderResourceTable,
 	const FVulkanBindingTable& BindingTable,
 	const TArray<VkDescriptorType>& DescriptorTypes,
-	const FVulkanSpirv& Spirv,
 	const TMap<FString, FVulkanShaderHeader::EType>& EntryTypes,
 	const FShaderCompilerInput& ShaderInput,
 	const FSpirvReflectBindings& SpirvReflectBindings,
 	FShaderParameterMap& InOutParameterMap,
-	FVulkanShaderHeader& OutHeader)
+	VulkanShaderCompilerSerializedOutput& SerializedOutput
+)
 {
+	FVulkanShaderHeader& OutHeader = SerializedOutput.Header;
+
 	// Names that match the Header.Globals array
 	TArray<FString> GlobalNames;
 
@@ -680,7 +693,7 @@ static void ConvertToHeader(
 			{
 			case FVulkanShaderHeader::Global:
 				{
-					const int32 HeaderGlobalIndex = AddGlobal(DescriptorTypes, ParameterName, BaseIndex, Spirv, OutHeader, GlobalNames);
+					const int32 HeaderGlobalIndex = AddGlobal(DescriptorTypes, ParameterName, BaseIndex, SerializedOutput, GlobalNames);
 
 					const FParameterAllocation* ParameterAllocation = InOutParameterMap.GetParameterMap().Find(*ParameterName);
 					check(ParameterAllocation);
@@ -704,7 +717,7 @@ static void ConvertToHeader(
 				}
 				break;
 			case FVulkanShaderHeader::UniformBuffer:
-				AddUniformBuffer(ShaderResourceTable, DescriptorTypes, ShaderInput, Spirv, ParameterName, BufferIndex, InOutParameterMap, OutHeader, GlobalNames);
+				AddUniformBuffer(ShaderResourceTable, DescriptorTypes, ShaderInput, ParameterName, BufferIndex, InOutParameterMap, SerializedOutput, GlobalNames);
 				break;
 			default:
 				check(0);
@@ -714,7 +727,7 @@ static void ConvertToHeader(
 		else
 		{
 			// Not found means it's a new resource-only UniformBuffer
-			AddUniformBuffer(ShaderResourceTable, DescriptorTypes, ShaderInput, Spirv, ParameterName, BufferIndex, InOutParameterMap, OutHeader, GlobalNames);
+			AddUniformBuffer(ShaderResourceTable, DescriptorTypes, ShaderInput, ParameterName, BufferIndex, InOutParameterMap, SerializedOutput, GlobalNames);
 		}
 	}
 
@@ -750,7 +763,7 @@ static void ConvertToHeader(
 				check(HeaderGlobalIndex != INDEX_NONE);
 				check(GlobalNames[HeaderGlobalIndex] == AttachmentName);
 				FVulkanShaderHeader::FGlobalInfo& GlobalInfo = OutHeader.Globals[HeaderGlobalIndex];
-				const FVulkanSpirv::FEntry* Entry = Spirv.GetEntry(AttachmentName);
+				const FVulkanSpirv::FEntry* Entry = SerializedOutput.Spirv.GetEntry(AttachmentName);
 				check(Entry);
 				check(Entry->Binding != -1);
 
@@ -772,7 +785,7 @@ static void ConvertToHeader(
 
 
 // Fills the SRT using final values kept in the FVulkanShaderHeader.
-// NOTE: Uses GloalIndex so it can be consumed directly at runtime.
+// NOTE: Uses GlobalIndex so it can be consumed directly at runtime.
 // NOTE: Keep in sync with BuildResourceTableMapping.
 static FShaderResourceTable BuildSRTFromHeader(const FVulkanShaderHeader& NEWHeader)
 {
@@ -833,13 +846,12 @@ static FShaderResourceTable BuildSRTFromHeader(const FVulkanShaderHeader& NEWHea
 
 
 static void BuildShaderOutput(
+	VulkanShaderCompilerSerializedOutput& SerializedOutput,
 	FShaderCompilerOutput&		ShaderOutput,
 	const FVulkanShaderCompilerInternalState& InternalState,
 	const FSpirvReflectBindings& SpirvReflectBindings,
 	TMap<FString, FVulkanShaderHeader::EType>& EntryTypes,
 	const FVulkanBindingTable&	BindingTable,
-	uint8						WaveSize,
-	const FVulkanSpirv&			Spirv,
 	const FString&				DebugName,
 	uint32						PackedGlobalArraySize,
 	TBitArray<>&				UsedUniformBufferSlots
@@ -848,10 +860,9 @@ static void BuildShaderOutput(
 	const FShaderCompilerInput& ShaderInput = InternalState.Input;
 	const EShaderFrequency Frequency = InternalState.GetShaderFrequency();
 
-	FVulkanShaderHeader NEWHeader(FVulkanShaderHeader::EZero);
-	FShaderCompilerResourceTable ShaderResourceTable;
+	FVulkanShaderHeader& NEWHeader = SerializedOutput.Header;
 
-	NEWHeader.SpirvCRC = Spirv.CRC;
+	NEWHeader.SpirvCRC = SerializedOutput.Spirv.CRC;
 	NEWHeader.RayTracingPayloadType = ShaderInput.Environment.GetCompileArgument(TEXT("RT_PAYLOAD_TYPE"), 0u);
 	NEWHeader.RayTracingPayloadSize = ShaderInput.Environment.GetCompileArgument(TEXT("RT_PAYLOAD_MAX_SIZE"), 0u);
 #if VULKAN_ENABLE_BINDING_DEBUG_NAMES
@@ -859,7 +870,7 @@ static void BuildShaderOutput(
 #endif
 
 	// :todo-jn: Hash entire SPIRV for now, could eventually be removed since we use ShaderKeys
-	FSHA1::HashBuffer(Spirv.Data.GetData(), Spirv.Data.Num(), (uint8*)&NEWHeader.SourceHash);
+	FSHA1::HashBuffer(SerializedOutput.Spirv.Data.GetData(), SerializedOutput.Spirv.GetByteSize(), (uint8*)&NEWHeader.SourceHash);
 
 
 	// Flattens the array dimensions of the interface variable (aka shader attribute), e.g. from float4[2][3] -> float4[6]
@@ -953,6 +964,7 @@ static void BuildShaderOutput(
 	ShaderOutput.ParameterMap.GetAllParameterNames(OriginalParameters);
 
 	// Build the SRT for this shader.
+	FShaderCompilerResourceTable ShaderResourceTable;
 	{
 		// Build the generic SRT for this shader.
 		FShaderCompilerResourceTable GenericSRT;
@@ -997,13 +1009,13 @@ static void BuildShaderOutput(
 		PackedUB.PackedTypeIndex = CrossCompiler::EPackedTypeIndex::HighP;
 		PackedUB.SizeInBytes = Align(PackedGlobalArraySize, 16u);
 
-		const FVulkanSpirv::FEntry* Entry = Spirv.GetEntryByBindingIndex(StageOffset);
+		const FVulkanSpirv::FEntry* Entry = SerializedOutput.Spirv.GetEntryByBindingIndex(StageOffset);
 		check(Entry);
 		PackedUB.SPIRVDescriptorSetOffset = Entry->WordDescriptorSetIndex;
 		PackedUB.SPIRVBindingIndexOffset = Entry->WordBindingIndex;
 	}
 
-	ConvertToHeader(ShaderResourceTable, BindingTable, DescriptorTypes, Spirv, EntryTypes, ShaderInput, SpirvReflectBindings, ShaderOutput.ParameterMap, NEWHeader);
+	ConvertToHeader(ShaderResourceTable, BindingTable, DescriptorTypes, EntryTypes, ShaderInput, SpirvReflectBindings, ShaderOutput.ParameterMap, SerializedOutput);
 	check(NEWHeader.EmulatedUBsCopyInfo.Num() == 0);
 
 	if (ShaderInput.Environment.CompilerFlags.Contains(CFLAG_ExtraShaderData))
@@ -1012,27 +1024,24 @@ static void BuildShaderOutput(
 	}
 
 	// Build the SRT for this shader from the NEWHeader
-	FShaderResourceTable SerializedSRT = BuildSRTFromHeader(NEWHeader);
-
-	// Plug the passed in WaveSize
-	NEWHeader.WaveSize = WaveSize;
+	SerializedOutput.ShaderResourceTable = BuildSRTFromHeader(NEWHeader);
 
 	// Write out the header and shader source code.
 	FMemoryWriter Ar(ShaderOutput.ShaderCode.GetWriteAccess(), true);
-	Ar << NEWHeader;
-	Ar << SerializedSRT;
+	Ar << SerializedOutput.Header;
+	Ar << SerializedOutput.ShaderResourceTable;
 
-	check(Spirv.Data.Num() != 0);
-	uint32 SpirvCodeSizeBytes = Spirv.Data.Num() * Spirv.Data.GetTypeSize();
+	check(SerializedOutput.Spirv.Data.Num() != 0);
+	uint32 SpirvCodeSizeBytes = SerializedOutput.Spirv.GetByteSize();
 	Ar << SpirvCodeSizeBytes;
-	Ar.Serialize((uint8*)Spirv.Data.GetData(), SpirvCodeSizeBytes);
+	Ar.Serialize((uint8*)SerializedOutput.Spirv.Data.GetData(), SpirvCodeSizeBytes);
 
 	ShaderOutput.bSucceeded = true;
 
 	if (ShaderInput.ExtraSettings.bExtractShaderSource)
 	{
 		TArray<ANSICHAR> AssemblyText;
-		if (CrossCompiler::FShaderConductorContext::Disassemble(CrossCompiler::EShaderConductorIR::Spirv, Spirv.GetByteData(), Spirv.GetByteSize(), AssemblyText))
+		if (CrossCompiler::FShaderConductorContext::Disassemble(CrossCompiler::EShaderConductorIR::Spirv, SerializedOutput.Spirv.GetByteData(), SerializedOutput.Spirv.GetByteSize(), AssemblyText))
 		{
 			ShaderOutput.OptionalFinalShaderSource = FString(AssemblyText.GetData());
 		}
@@ -1041,7 +1050,7 @@ static void BuildShaderOutput(
 	{
 		if (SupportsOfflineCompiler(ShaderInput.ShaderFormat))
 		{
-			CompileOfflineMali(ShaderInput, ShaderOutput, (const ANSICHAR*)Spirv.GetByteData(), Spirv.GetByteSize(), true, Spirv.EntryPointName);
+			CompileOfflineMali(ShaderInput, ShaderOutput, (const ANSICHAR*)SerializedOutput.Spirv.GetByteData(), SerializedOutput.Spirv.GetByteSize(), true, SerializedOutput.Spirv.EntryPointName);
 		}
 	}
 
@@ -1186,11 +1195,10 @@ static uint32 CalculateSpirvInstructionCount(FVulkanSpirv& Spirv)
 
 static bool BuildShaderOutputFromSpirv(
 	CrossCompiler::FShaderConductorContext&	CompilerContext,
-	FVulkanSpirv&							Spirv,
 	const FVulkanShaderCompilerInternalState& InternalState,
+	VulkanShaderCompilerSerializedOutput&   SerializedOutput,
 	FShaderCompilerOutput&					Output,
-	FVulkanBindingTable&					BindingTable,
-	uint8									WaveSize
+	FVulkanBindingTable&					BindingTable
 )
 {
 	const bool bIsRayTracingShader = InternalState.IsRayTracingShader();
@@ -1198,8 +1206,8 @@ static bool BuildShaderOutputFromSpirv(
 	FShaderParameterMap& ParameterMap = Output.ParameterMap;
 
 	// Reflect SPIR-V module with SPIRV-Reflect library
-	const size_t SpirvDataSize = Spirv.GetByteSize();
-	spv_reflect::ShaderModule Reflection(SpirvDataSize, Spirv.GetByteData(), SPV_REFLECT_RETURN_FLAG_SAMPLER_IMAGE_USAGE);
+	const size_t SpirvDataSize = SerializedOutput.Spirv.GetByteSize();
+	spv_reflect::ShaderModule Reflection(SpirvDataSize, SerializedOutput.Spirv.GetByteData(), SPV_REFLECT_RETURN_FLAG_SAMPLER_IMAGE_USAGE);
 	check(Reflection.GetResult() == SPV_REFLECT_RESULT_SUCCESS);
 
 	// Ray tracing shaders are not being rewritten to remove unreferenced entry points due to a bug in dxc.
@@ -1291,7 +1299,8 @@ static bool BuildShaderOutputFromSpirv(
 		{
 			for (const SpvReflectDescriptorBinding* Binding : BindingArray)
 			{
-				checkf(!InternalState.bSupportsBindless || (BindingType == EVulkanBindingType::UniformBuffer), TEXT("Bindless shaders should only have uniform buffers."));
+				checkf(!InternalState.bSupportsBindless || (BindingType == EVulkanBindingType::UniformBuffer) || (BindingType == EVulkanBindingType::PackedUniformBuffer),
+					TEXT("Bindless shaders should only have uniform buffers."));
 
 				const FString ResourceName(ANSI_TO_TCHAR(Binding->name));
 
@@ -1306,7 +1315,7 @@ static bool BuildShaderOutputFromSpirv(
 				const SpvReflectResult SpvResult = Reflection.ChangeDescriptorBindingNumbers(Binding, BindingIndex, DescSetNumber);
 				check(SpvResult == SPV_REFLECT_RESULT_SUCCESS);
 
-				const int32 VulkanBindingIndex = Spirv.ReflectionInfo.Add(FVulkanSpirv::FEntry(ResourceName, BindingIndex));
+				const int32 VulkanBindingIndex = SerializedOutput.Spirv.ReflectionInfo.Add(FVulkanSpirv::FEntry(ResourceName, BindingIndex));
 
 				switch (BindingType)
 				{
@@ -1418,13 +1427,13 @@ static bool BuildShaderOutputFromSpirv(
 	Output.Target = InternalState.Input.Target;
 
 	// Overwrite updated SPIRV code
-	Spirv.Data = TArray<uint32>(Reflection.GetCode(), Reflection.GetCodeSize() / 4);
+	SerializedOutput.Spirv.Data = TArray<uint32>(Reflection.GetCode(), Reflection.GetCodeSize() / 4);
 
 	// We have to strip out most debug instructions (except OpName) for Vulkan mobile
 	if (InternalState.bStripReflect)
 	{
 		const char* OptArgs[] = { "--strip-reflect", "-O"};
-		if (!CompilerContext.OptimizeSpirv(Spirv.Data, OptArgs, UE_ARRAY_COUNT(OptArgs)))
+		if (!CompilerContext.OptimizeSpirv(SerializedOutput.Spirv.Data, OptArgs, UE_ARRAY_COUNT(OptArgs)))
 		{
 			UE_LOG(LogVulkanShaderCompiler, Error, TEXT("Failed to strip debug instructions from SPIR-V module"));
 			return false;
@@ -1435,27 +1444,26 @@ static bool BuildShaderOutputFromSpirv(
 	if (IsAndroidShaderFormat(InternalState.Input.ShaderFormat))
 	{
 		const char* OptArgs[] = { "--android-driver-patch" };
-		if (!CompilerContext.OptimizeSpirv(Spirv.Data, OptArgs, UE_ARRAY_COUNT(OptArgs)))
+		if (!CompilerContext.OptimizeSpirv(SerializedOutput.Spirv.Data, OptArgs, UE_ARRAY_COUNT(OptArgs)))
 		{
 			UE_LOG(LogVulkanShaderCompiler, Error, TEXT("Failed to apply driver patches for Android"));
 			return false;
 		}
 	}
 
-	PatchSpirvReflectionEntries(Spirv);
+	PatchSpirvReflectionEntries(SerializedOutput.Spirv);
 
-	Spirv.EntryPointName = PatchSpirvEntryPointWithCRC(Spirv, Spirv.CRC);
+	SerializedOutput.Spirv.EntryPointName = PatchSpirvEntryPointWithCRC(SerializedOutput.Spirv, SerializedOutput.Spirv.CRC);
 
-	Output.NumInstructions = CalculateSpirvInstructionCount(Spirv);
+	Output.NumInstructions = CalculateSpirvInstructionCount(SerializedOutput.Spirv);
 
 	BuildShaderOutput(
+		SerializedOutput,
 		Output,
 		InternalState,
 		Bindings,
 		EntryTypes,
 		BindingTable,
-		WaveSize,
-		Spirv,
 		InternalState.GetDebugName(),
 		PackedGlobalArraySize,
 		UsedUniformBufferSlots
@@ -1464,8 +1472,8 @@ static bool BuildShaderOutputFromSpirv(
 	if (InternalState.bDebugDump)
 	{
 		// Write meta data to debug output file and write SPIR-V dump in binary and text form
-		DumpDebugShaderBinary(InternalState.Input, Spirv.GetByteData(), Spirv.GetByteSize(), TEXT("spv"));
-		DumpDebugShaderDisassembledSpirv(InternalState.Input, Spirv.GetByteData(), Spirv.GetByteSize(), TEXT("spvasm"));
+		DumpDebugShaderBinary(InternalState.Input, SerializedOutput.Spirv.GetByteData(), SerializedOutput.Spirv.GetByteSize(), TEXT("spv"));
+		DumpDebugShaderDisassembledSpirv(InternalState.Input, SerializedOutput.Spirv.GetByteData(), SerializedOutput.Spirv.GetByteSize(), TEXT("spvasm"));
 	}
 
 	return true;
@@ -1853,16 +1861,13 @@ uint8 ParseWaveSize(
 static bool CompileWithShaderConductor(
 	const FVulkanShaderCompilerInternalState& InternalState,
 	const FString&			PreprocessedShader,
-	FShaderCompilerOutput&	Output,
-	FVulkanBindingTable&	BindingTable
+	VulkanShaderCompilerSerializedOutput& SerializedOutput,
+	FShaderCompilerOutput&	Output
 )
 {
 	const FShaderCompilerInput& Input = InternalState.Input;
 
-	const bool bIsRayTracingShader = InternalState.IsRayTracingShader();
-	const bool bRewriteHlslSource = !bIsRayTracingShader;
-	const bool bDebugDump = InternalState.bDebugDump;
-
+	FVulkanBindingTable BindingTable(InternalState.GetHlslShaderFrequency());
 	CrossCompiler::FShaderConductorContext CompilerContext;
 
 	// Inject additional macro definitions to circumvent missing features: external textures
@@ -1878,7 +1883,7 @@ static bool CompileWithShaderConductor(
 	Options.TargetEnvironment = InternalState.MinimumTargetEnvironment;
 
 	// VK_EXT_scalar_block_layout is required by raytracing and by Nanite (so expect it to be present in SM6/Vulkan_1_3)
-	Options.bDisableScalarBlockLayout = !(bIsRayTracingShader || 
+	Options.bDisableScalarBlockLayout = !(InternalState.IsRayTracingShader() ||
 		InternalState.MinimumTargetEnvironment >= CrossCompiler::FShaderConductorOptions::ETargetEnvironment::Vulkan_1_3);
 
 	if (Input.Environment.CompilerFlags.Contains(CFLAG_AllowRealTypes))
@@ -1892,36 +1897,36 @@ static bool CompileWithShaderConductor(
 		Options.HlslVersion = 2021;
 	}
 
-	if (bDebugDump)
+	if (InternalState.bDebugDump)
 	{
 		VulkanCreateDXCCompileBatchFiles(InternalState, Options);
 	}
 
 	// Before the shader rewritter removes all traces of it, pull any WAVESIZE directives from the shader source
-	const uint8 WaveSize = ParseWaveSize(InternalState, PreprocessedShader);
+	SerializedOutput.Header.WaveSize = ParseWaveSize(InternalState, PreprocessedShader);
 
+	const bool bRewriteHlslSource = !InternalState.IsRayTracingShader();
 	if (bRewriteHlslSource)
 	{
 		// Rewrite HLSL source code to remove unused global resources and variables
 		FString RewrittenHlslSource;
 
 		Options.bRemoveUnusedGlobals = true;
-		if (!CompilerContext.RewriteHlsl(Options, (bDebugDump ? &RewrittenHlslSource : nullptr)))
+		if (!CompilerContext.RewriteHlsl(Options, (InternalState.bDebugDump ? &RewrittenHlslSource : nullptr)))
 		{
 			CompilerContext.FlushErrors(Output.Errors);
 			return false;
 		}
 		Options.bRemoveUnusedGlobals = false;
 
-		if (bDebugDump)
+		if (InternalState.bDebugDump)
 		{
 			DumpDebugShaderText(Input, RewrittenHlslSource, TEXT("rewritten.hlsl"));
 		}
 	}
 
 	// Compile HLSL source to SPIR-V binary
-	FVulkanSpirv Spirv;
-	if (!CompilerContext.CompileHlslToSpirv(Options, Spirv.Data))
+	if (!CompilerContext.CompileHlslToSpirv(Options, SerializedOutput.Spirv.Data))
 	{
 		CompilerContext.FlushErrors(Output.Errors);
 		return false;
@@ -1929,15 +1934,10 @@ static bool CompileWithShaderConductor(
 
 	// If this shader samples R64 image formats, they need to get converted to STORAGE_IMAGE
 	// todo-jnmo: Scope this with a CFLAG if it affects compilation times 
-	Patch64bitSamplers(Spirv);
+	Patch64bitSamplers(SerializedOutput.Spirv);
 
 	// Build shader output and binding table
-	Output.bSucceeded = BuildShaderOutputFromSpirv(CompilerContext, Spirv, InternalState, Output, BindingTable, WaveSize);
-
-	if (Input.Environment.CompilerFlags.Contains(CFLAG_ExtraShaderData))
-	{
-		Output.ShaderCode.AddOptionalData(FShaderCodeName::Key, TCHAR_TO_UTF8(*Input.GenerateShaderName()));
-	}
+	Output.bSucceeded = BuildShaderOutputFromSpirv(CompilerContext, InternalState, SerializedOutput, Output, BindingTable);
 
 	// Flush warnings
 	CompilerContext.FlushErrors(Output.Errors);
@@ -2049,7 +2049,7 @@ bool PreprocessVulkanShader(const FShaderCompilerInput& Input, const FShaderComp
 	#if UE_VULKAN_SHADER_COMPILER_ALLOW_DEAD_CODE_REMOVAL
 	if (Input.Environment.CompilerFlags.Contains(CFLAG_RemoveDeadCode))
 	{
-		UE::ShaderCompilerCommon::RemoveDeadCode(PreprocessedShaderSource, Input.EntryPointName, PreprocessOutput.EditErrors());
+		UE::ShaderCompilerCommon::RemoveDeadCode(PreprocessedShaderSource, InternalState.GetEntryPointName(), PreprocessOutput.EditErrors());
 	}
 	#endif // UE_VULKAN_SHADER_COMPILER_ALLOW_DEAD_CODE_REMOVAL
 
@@ -2073,15 +2073,19 @@ void CompileVulkanShader(const FShaderCompilerInput& Input, const FShaderPreproc
 		return;
 	}
 
-	FVulkanBindingTable BindingTable(HlslFrequency);
-
 	bool bSuccess = false;
 
 #if PLATFORM_MAC || PLATFORM_WINDOWS || PLATFORM_LINUX
-	// Cross-compile shader via ShaderConductor (DXC, SPIRV-Tools, SPIRV-Cross)
-	bSuccess = CompileWithShaderConductor(InternalState, PreprocessOutput.GetSource(), Output, BindingTable);
+	// Compile shader via ShaderConductor (DXC, SPIRV-Tools)
+	VulkanShaderCompilerSerializedOutput SerializedOutput;
+	bSuccess = CompileWithShaderConductor(InternalState, PreprocessOutput.GetSource(), SerializedOutput, Output);
 #endif // PLATFORM_MAC || PLATFORM_WINDOWS || PLATFORM_LINUX
 	
+	if (Input.Environment.CompilerFlags.Contains(CFLAG_ExtraShaderData))
+	{
+		Output.ShaderCode.AddOptionalData(FShaderCodeName::Key, TCHAR_TO_UTF8(*Input.GenerateShaderName()));
+	}
+
 	PreprocessOutput.GetParameterParser().ValidateShaderParameterTypes(Input, InternalState.IsMobileES31(), Output);
 	
 	const bool bDirectCompile = FParse::Param(FCommandLine::Get(), TEXT("directcompile"));
