VulkanRHI: Start adding code in VulkanShaderCompiler to track bindless UB and handle RootParametersStructure. (VK RT WIP, disabled for now)

#rnx

diff --git a/Engine/Source/Developer/VulkanShaderFormat/Private/VulkanShaderCompiler.cpp b/Engine/Source/Developer/VulkanShaderFormat/Private/VulkanShaderCompiler.cpp
--- a/Engine/Source/Developer/VulkanShaderFormat/Private/VulkanShaderCompiler.cpp
+++ b/Engine/Source/Developer/VulkanShaderFormat/Private/VulkanShaderCompiler.cpp
@@ -227,6 +227,20 @@ struct FPatchType
 };
 
 
+static const FString kBindlessCBPrefix = TEXT("__BindlessCB");
+static const FString kBindlessHeapSuffix = TEXT("_Heap");
+static FString GetBindlessUBNameFromHeap(const FString& HeapName)
+{
+	check(HeapName.StartsWith(kBindlessCBPrefix));
+	check(HeapName.EndsWith(kBindlessHeapSuffix));
+
+	int32 NameStart = HeapName.Find(TEXT("_"), ESearchCase::IgnoreCase, ESearchDir::FromStart, kBindlessCBPrefix.Len() + 1);
+	check(NameStart != INDEX_NONE);
+	NameStart++;
+	return HeapName.Mid(NameStart, HeapName.Len() - NameStart - kBindlessHeapSuffix.Len());
+}
+
+
 // A collection of states and data that is locked in at the top level call and doesn't change throughout the compilation process
 struct FVulkanShaderCompilerInternalState
 {
@@ -238,7 +252,6 @@ struct FVulkanShaderCompilerInternalState
 		, bUseBindlessUniformBuffer(InInput.IsRayTracingShader() && ((EShaderFrequency)InInput.Target.Frequency != SF_RayGen))
 		, bIsRayHitGroupShader(InInput.IsRayTracingShader() && ((EShaderFrequency)InInput.Target.Frequency == SF_RayHitGroup))
 		, bSupportsBindless(InInput.Environment.CompilerFlags.Contains(CFLAG_BindlessResources) || InInput.Environment.CompilerFlags.Contains(CFLAG_BindlessSamplers))
-		, bUseRootParameterStructure(InInput.IsRayTracingShader() && InInput.RootParametersStructure)
 		, bDebugDump(InInput.DumpDebugInfoEnabled())
 	{
 	}
@@ -253,9 +266,10 @@ struct FVulkanShaderCompilerInternalState
 	const bool bIsRayHitGroupShader;
 
 	const bool bSupportsBindless;
-	const bool bUseRootParameterStructure;
 	const bool bDebugDump;
 
+	TArray<FString> AllBindlessUBs;
+
 	// Forwarded calls for convenience
 	inline EShaderFrequency GetShaderFrequency() const
 	{
@@ -269,6 +283,11 @@ struct FVulkanShaderCompilerInternalState
 	{
 		return Input.IsRayTracingShader();
 	}
+	inline bool UseRootParametersStructure() const
+	{
+		// Only supported for RayGen currently
+		return (GetShaderFrequency() == SF_RayGen) && (Input.RootParametersStructure != nullptr);
+	}
 	inline bool IsSM6() const
 	{
 		return (Version == EVulkanShaderVersion::SM6);
@@ -315,6 +334,8 @@ struct VulkanShaderCompilerSerializedOutput
 	FVulkanShaderHeader Header;
 	FShaderResourceTable ShaderResourceTable;
 	FVulkanSpirv Spirv;
+
+	TSet<FString> UsedBindlessUB;
 };
 
 
@@ -1063,16 +1084,17 @@ static void BuildShaderOutput(
 static void GatherSpirvReflectionBindings(
 	spv_reflect::ShaderModule&	Reflection,
 	FSpirvReflectBindings&		OutBindings,
-	const EShaderFrequency		ShaderFrequency,
-	const bool					bSupportsBindless)
+	TSet<FString>&				OutBindlessUB,
+	const FVulkanShaderCompilerInternalState& InternalState)
 {
 	// Change descriptor set numbers
 	TArray<SpvReflectDescriptorSet*> DescriptorSets;
 	uint32 NumDescriptorSets = 0;
 
 	// If bindless is supported, then offset the descriptor set to fit the bindless heaps at the beginning
+	const EShaderFrequency ShaderFrequency = InternalState.GetShaderFrequency();
 	const uint32 StageIndex = (uint32)ShaderStage::GetStageForFrequency(ShaderFrequency);
-	const uint32 DescSetNo = bSupportsBindless ? VulkanBindless::NumBindlessSets + StageIndex : StageIndex;
+	const uint32 DescSetNo = InternalState.bSupportsBindless ? VulkanBindless::NumBindlessSets + StageIndex : StageIndex;
 
 	SpvReflectResult SpvResult = Reflection.EnumerateDescriptorSets(&NumDescriptorSets, nullptr);
 	check(SpvResult == SPV_REFLECT_RESULT_SUCCESS);
@@ -1103,7 +1125,7 @@ static void GatherSpirvReflectionBindings(
 	}
 
 	// Patch resource heaps descriptor set numbers
-	if (bSupportsBindless)
+	if (InternalState.bSupportsBindless)
 	{
 		// Move the bindless heap to its dedicated descriptor set and remove it from our regular binding arrays
 		auto MoveBindlessHeaps = [&](TArray<SpvReflectDescriptorBinding*>& BindingArray, const TCHAR* HeapPrefix, uint32 BinldessDescSetNo)
@@ -1138,7 +1160,20 @@ static void GatherSpirvReflectionBindings(
 			for (int32 Index = OutBindings.UniformBuffers.Num() - 1; Index >= 0; --Index)
 			{
 				const SpvReflectDescriptorBinding* pBinding = OutBindings.UniformBuffers[Index];
-				Reflection.ChangeDescriptorBindingNumbers(pBinding, BindingOffset + pBinding->binding, VulkanBindless::BindlessSingleUseUniformBufferSet);
+				const FString BindingName(ANSI_TO_TCHAR(pBinding->name));
+				if (BindingName.StartsWith(kBindlessCBPrefix))
+				{
+					check(InternalState.bUseBindlessUniformBuffer);
+					Reflection.ChangeDescriptorBindingNumbers(pBinding, 0, VulkanBindless::BindlessUniformBufferSet);
+					const FString BindlessUBName = GetBindlessUBNameFromHeap(BindingName);
+					checkf(InternalState.AllBindlessUBs.Contains(BindlessUBName), TEXT("Bindless Uniform Buffer was found in SPIRV but not tracked in internal state"));
+					OutBindlessUB.Add(BindlessUBName);
+					OutBindings.UniformBuffers.RemoveAtSwap(Index);
+				}
+				else
+				{
+					Reflection.ChangeDescriptorBindingNumbers(pBinding, BindingOffset + pBinding->binding, VulkanBindless::BindlessSingleUseUniformBufferSet);
+				}
 			}
 		}
 	}
@@ -1202,7 +1237,6 @@ static bool BuildShaderOutputFromSpirv(
 )
 {
 	const bool bIsRayTracingShader = InternalState.IsRayTracingShader();
-	const EShaderFrequency ShaderFrequency = InternalState.GetShaderFrequency();
 	FShaderParameterMap& ParameterMap = Output.ParameterMap;
 
 	// Reflect SPIR-V module with SPIRV-Reflect library
@@ -1238,19 +1272,22 @@ static bool BuildShaderOutputFromSpirv(
 	}
 
 	FSpirvReflectBindings Bindings;
-	GatherSpirvReflectionBindings(Reflection, Bindings, ShaderFrequency, InternalState.bSupportsBindless);
+	GatherSpirvReflectionBindings(Reflection, Bindings, SerializedOutput.UsedBindlessUB, InternalState);
 
 	// Build binding table
 	TMap<const SpvReflectDescriptorBinding*, int32> BindingToIndexMap;
 
-	const FString UBOGlobalsNameSpv = ANSI_TO_TCHAR(CrossCompiler::FShaderConductorContext::GetIdentifierTable().GlobalsUniformBuffer);
+	const FString UBOGlobalsNameSpv(ANSI_TO_TCHAR(CrossCompiler::FShaderConductorContext::GetIdentifierTable().GlobalsUniformBuffer));
+	const FString UBORootParamNameSpv(FShaderParametersMetadata::kRootUniformBufferBindingName);
 
-	auto RegisterBindings = [&BindingTable, &BindingToIndexMap, &UBOGlobalsNameSpv](TArray<SpvReflectDescriptorBinding*>& Bindings, const char* BlockName, EVulkanBindingType::EType BindingType)
+	auto RegisterBindings = [&BindingTable, &BindingToIndexMap, &UBOGlobalsNameSpv, &UBORootParamNameSpv]
+							(TArray<SpvReflectDescriptorBinding*>& Bindings, const char* BlockName, EVulkanBindingType::EType BindingType)
 	{
 		for (const SpvReflectDescriptorBinding* Binding : Bindings)
 		{
-			if (((BindingType == EVulkanBindingType::PackedUniformBuffer) && (UBOGlobalsNameSpv != Binding->name)) ||
-				((BindingType == EVulkanBindingType::UniformBuffer) && (UBOGlobalsNameSpv == Binding->name)))
+			const bool bIsGlobalOrRootBuffer = ((UBOGlobalsNameSpv == Binding->name) || (UBORootParamNameSpv == Binding->name));
+			if (((BindingType == EVulkanBindingType::PackedUniformBuffer) && !bIsGlobalOrRootBuffer) ||
+				((BindingType == EVulkanBindingType::UniformBuffer) && bIsGlobalOrRootBuffer))
 			{
 				continue;
 			}
@@ -1292,7 +1329,7 @@ static bool BuildShaderOutputFromSpirv(
 
 	// Final descriptor binding numbers for all other resource types
 	{
-		const int32 StageOffset = InternalState.bSupportsBindless ? (ShaderStage::GetStageForFrequency(ShaderFrequency) * VulkanBindless::MaxUniformBuffersPerStage) : 0;
+		const int32 StageOffset = InternalState.bSupportsBindless ? (ShaderStage::GetStageForFrequency(InternalState.GetShaderFrequency()) * VulkanBindless::MaxUniformBuffersPerStage) : 0;
 		const uint32_t DescSetNumber = InternalState.bSupportsBindless ? (uint32_t)VulkanBindless::BindlessSingleUseUniformBufferSet : (uint32_t)SPV_REFLECT_SET_NUMBER_DONT_CHANGE;
 
 		auto AddReflectionInfos = [&](TArray<SpvReflectDescriptorBinding*>& BindingArray, EVulkanBindingType::EType BindingType, int32 BindingOffset)
@@ -1304,8 +1341,9 @@ static bool BuildShaderOutputFromSpirv(
 
 				const FString ResourceName(ANSI_TO_TCHAR(Binding->name));
 
-				if (((BindingType == EVulkanBindingType::PackedUniformBuffer) && (ResourceName != UBOGlobalsNameSpv)) ||
-					((BindingType == EVulkanBindingType::UniformBuffer) && (ResourceName == UBOGlobalsNameSpv)))
+				const bool bIsGlobalOrRootBuffer = ((UBOGlobalsNameSpv == ResourceName) || (UBORootParamNameSpv == ResourceName));
+				if (((BindingType == EVulkanBindingType::PackedUniformBuffer) && !bIsGlobalOrRootBuffer) ||
+					((BindingType == EVulkanBindingType::UniformBuffer) && bIsGlobalOrRootBuffer))
 				{
 					continue;
 				}
@@ -1360,8 +1398,13 @@ static bool BuildShaderOutputFromSpirv(
 						// Use the given global ResourceName instead of patching it to _Globals_h
 						check(!UsedUniformBufferSlots[VulkanBindingIndex]);
 						UsedUniformBufferSlots[VulkanBindingIndex] = true;
-						//HandleReflectedUniformBuffer(ResourceName, VulkanBindingIndex, Output);
-						//EntryTypes.Add(ResourceName, FVulkanShaderHeader::UniformBuffer);
+
+						if (InternalState.UseRootParametersStructure())
+						{
+							check(VulkanBindingIndex == FShaderParametersMetadata::kRootCBufferBindingIndex);
+							HandleReflectedUniformBuffer(ResourceName, VulkanBindingIndex, Output);
+							EntryTypes.Add(ResourceName, FVulkanShaderHeader::UniformBuffer);
+						}
 
 						// Register all uniform buffer members of Globals as loose data
 						for (uint32 MemberIndex = 0; MemberIndex < Binding->block.member_count; ++MemberIndex)
@@ -1369,16 +1412,22 @@ static bool BuildShaderOutputFromSpirv(
 							const SpvReflectBlockVariable& Member = Binding->block.members[MemberIndex];
 
 							FString MemberName(ANSI_TO_TCHAR(Member.name));
-							HandleReflectedGlobalConstantBufferMember(
-								MemberName,
-								BindingOffset,
-								Member.absolute_offset,
-								Member.size,
-								Output
-							);
-
-							FShaderParameterParser::RemoveBindlessParameterPrefix(MemberName);
-							EntryTypes.Add(MemberName, FVulkanShaderHeader::PackedGlobal);
+
+							// Add all members of global ub, and only bindless samplers/resources for root param
+							if (!InternalState.UseRootParametersStructure() ||
+								(MemberName.StartsWith(FShaderParameterParser::kBindlessSamplerPrefix) || MemberName.StartsWith(FShaderParameterParser::kBindlessResourcePrefix)))
+							{
+								HandleReflectedGlobalConstantBufferMember(
+									MemberName,
+									BindingOffset,
+									Member.absolute_offset,
+									Member.size,
+									Output
+								);
+
+								FShaderParameterParser::RemoveBindlessParameterPrefix(MemberName);
+								EntryTypes.Add(MemberName, FVulkanShaderHeader::PackedGlobal);
+							}
 
 							PackedGlobalArraySize = FMath::Max<uint32>((Member.absolute_offset + Member.size), PackedGlobalArraySize);
 						}
