VulkanRHI: Add handling of multiple entrypoints for RayHitGroup shaders (VK RT WIP, does not change current output)

#rnx

diff --git a/Engine/Source/Developer/VulkanShaderFormat/Private/VulkanShaderCompiler.cpp b/Engine/Source/Developer/VulkanShaderFormat/Private/VulkanShaderCompiler.cpp
--- a/Engine/Source/Developer/VulkanShaderFormat/Private/VulkanShaderCompiler.cpp
+++ b/Engine/Source/Developer/VulkanShaderFormat/Private/VulkanShaderCompiler.cpp
@@ -240,6 +240,8 @@ static FString GetBindlessUBNameFromHeap(const FString& HeapName)
 	return HeapName.Mid(NameStart, HeapName.Len() - NameStart - kBindlessHeapSuffix.Len());
 }
 
+// Name of the structure in raytracing shader records in VulkanCommon.usf
+static const FString kHitGroupSystemRootConstantsSymbolName = TEXT("HitGroupSystemRootConstants");
 
 // A collection of states and data that is locked in at the top level call and doesn't change throughout the compilation process
 struct FVulkanShaderCompilerInternalState
@@ -254,6 +256,11 @@ struct FVulkanShaderCompilerInternalState
 		, bSupportsBindless(InInput.Environment.CompilerFlags.Contains(CFLAG_BindlessResources) || InInput.Environment.CompilerFlags.Contains(CFLAG_BindlessSamplers))
 		, bDebugDump(InInput.DumpDebugInfoEnabled())
 	{
+		if (bIsRayHitGroupShader)
+		{
+			UE::ShaderCompilerCommon::ParseRayTracingEntryPoint(Input.EntryPointName, ClosestHitEntry, AnyHitEntry, IntersectionEntry);
+			checkf(!ClosestHitEntry.IsEmpty(), TEXT("All hit groups must contain at least a closest hit shader module"));
+		}
 	}
 
 	const FShaderCompilerInput& Input;
@@ -268,6 +275,19 @@ struct FVulkanShaderCompilerInternalState
 	const bool bSupportsBindless;
 	const bool bDebugDump;
 
+	// Ray tracing specific states
+	enum class EHitGroupShaderType
+	{
+		None,
+		ClosestHit,
+		AnyHit,
+		Intersection
+	};
+	EHitGroupShaderType HitGroupShaderType = EHitGroupShaderType::None;
+	FString ClosestHitEntry;
+	FString AnyHitEntry;
+	FString IntersectionEntry;
+
 	TArray<FString> AllBindlessUBs;
 
 	// Forwarded calls for convenience
@@ -277,7 +297,59 @@ struct FVulkanShaderCompilerInternalState
 	}
 	inline const FString& GetEntryPointName() const
 	{
-		return Input.EntryPointName;
+		if (bIsRayHitGroupShader)
+		{
+			switch (HitGroupShaderType)
+			{
+			case EHitGroupShaderType::AnyHit: 
+				return AnyHitEntry;
+			case EHitGroupShaderType::Intersection:
+				return IntersectionEntry;
+			case EHitGroupShaderType::ClosestHit:
+				return ClosestHitEntry;
+
+			case EHitGroupShaderType::None:
+				[[fallthrough]];
+			default:
+				return Input.EntryPointName;
+			};
+		}
+		else
+		{
+			return Input.EntryPointName;
+		}
+	}
+	inline TArray<FStringView> GetRequiredSymbols() const
+	{
+		TArray<FStringView> RequiredSymbols;
+
+		if (bIsRayHitGroupShader)
+		{
+			check(!ClosestHitEntry.IsEmpty());
+
+			RequiredSymbols.Add(ClosestHitEntry);
+
+			if (!AnyHitEntry.IsEmpty())
+			{
+				RequiredSymbols.Add(AnyHitEntry);
+			}
+
+			if (!IntersectionEntry.IsEmpty())
+			{
+				RequiredSymbols.Add(IntersectionEntry);
+			}
+		}
+		else
+		{
+			RequiredSymbols.Add(Input.EntryPointName);
+		}
+
+		if (IsRayTracingShader())
+		{
+			RequiredSymbols.Add(kHitGroupSystemRootConstantsSymbolName);
+		}
+
+		return RequiredSymbols;
 	}
 	inline bool IsRayTracingShader() const
 	{
@@ -321,6 +393,26 @@ struct FVulkanShaderCompilerInternalState
 	{
 		return Input.DumpDebugInfoPath.Right(Input.DumpDebugInfoPath.Len() - Input.DumpDebugInfoRootPath.Len());
 	}
+	inline bool HasMultipleEntryPoints() const
+	{
+		return !ClosestHitEntry.IsEmpty() && (!AnyHitEntry.IsEmpty() || !IntersectionEntry.IsEmpty());
+	}
+	inline FString GetSPVExtension() const
+	{
+		switch (HitGroupShaderType)
+		{
+		case EHitGroupShaderType::AnyHit:
+			return TEXT("anyhit.spv");
+		case EHitGroupShaderType::Intersection:
+			return TEXT("intersection.spv");
+		case EHitGroupShaderType::ClosestHit:
+			return TEXT("closesthit.spv");
+		case EHitGroupShaderType::None: 
+			[[fallthrough]];
+		default:
+			return TEXT("spv");
+		};
+	}
 };
 
 // Data structures that will get serialized into ShaderCompilerOutput
@@ -1075,7 +1167,11 @@ static void BuildShaderOutput(
 		}
 	}
 
-	CullGlobalUniformBuffers(ShaderInput.Environment.UniformBufferMap, ShaderOutput.ParameterMap);
+	// Ray generation shaders rely on a different binding model that aren't compatible with global uniform buffers.
+	if (!InternalState.IsRayTracingShader())
+	{
+		CullGlobalUniformBuffers(ShaderInput.Environment.UniformBufferMap, ShaderOutput.ParameterMap);
+	}
 }
 
 
@@ -1236,38 +1332,17 @@ static bool BuildShaderOutputFromSpirv(
 	FVulkanBindingTable&					BindingTable
 )
 {
-	const bool bIsRayTracingShader = InternalState.IsRayTracingShader();
-	FShaderParameterMap& ParameterMap = Output.ParameterMap;
-
 	// Reflect SPIR-V module with SPIRV-Reflect library
 	const size_t SpirvDataSize = SerializedOutput.Spirv.GetByteSize();
 	spv_reflect::ShaderModule Reflection(SpirvDataSize, SerializedOutput.Spirv.GetByteData(), SPV_REFLECT_RETURN_FLAG_SAMPLER_IMAGE_USAGE);
 	check(Reflection.GetResult() == SPV_REFLECT_RESULT_SUCCESS);
 
 	// Ray tracing shaders are not being rewritten to remove unreferenced entry points due to a bug in dxc.
-	// Until it's fixed and integrated, we need to pull out the requested entry point manually.
-	int32 EntryPointIndex = (!bIsRayTracingShader) ? 0 : -1;
-	if (bIsRayTracingShader)
-	{
-		// For now only use the primary entry point for hit groups until we decide how to best support hit group library shaders.
-		FString OutMain, OutAnyHit, OutIntersection;
-		UE::ShaderCompilerCommon::ParseRayTracingEntryPoint(InternalState.GetEntryPointName(), OutMain, OutAnyHit, OutIntersection);
-
-		for (uint32 i = 0; i < Reflection.GetEntryPointCount(); ++i)
-		{
-			if (OutMain.Equals(Reflection.GetEntryPointName(i)))
-			{
-				EntryPointIndex = static_cast<int32>(i);
-				break;
-			}
-		}
-		checkf(EntryPointIndex >= 0, TEXT("Failed to find entry point %s in SPIRV-V module."), *InternalState.GetEntryPointName());
-	}
-
+	// An issue prevents multiple entrypoints in the same spirv module, so limit ourselves to one entrypoint at a time
 	// Change final entry point name in SPIR-V module
 	{
-		checkf(bIsRayTracingShader || Reflection.GetEntryPointCount() == 1, TEXT("Too many entry points in SPIR-V module: Expected 1, but got %d"), Reflection.GetEntryPointCount());
-		SpvReflectResult Result = Reflection.ChangeEntryPointName(EntryPointIndex, "main_00000000_00000000");
+		checkf(Reflection.GetEntryPointCount() == 1, TEXT("Too many entry points in SPIR-V module: Expected 1, but got %d"), Reflection.GetEntryPointCount());
+		const SpvReflectResult Result = Reflection.ChangeEntryPointName(0, "main_00000000_00000000");
 		check(Result == SPV_REFLECT_RESULT_SUCCESS);
 	}
 
@@ -1520,9 +1595,12 @@ static bool BuildShaderOutputFromSpirv(
 
 	if (InternalState.bDebugDump)
 	{
+		FString SPVExt(InternalState.GetSPVExtension());
+		FString SPVASMExt(SPVExt + TEXT("asm"));
+
 		// Write meta data to debug output file and write SPIR-V dump in binary and text form
-		DumpDebugShaderBinary(InternalState.Input, SerializedOutput.Spirv.GetByteData(), SerializedOutput.Spirv.GetByteSize(), TEXT("spv"));
-		DumpDebugShaderDisassembledSpirv(InternalState.Input, SerializedOutput.Spirv.GetByteData(), SerializedOutput.Spirv.GetByteSize(), TEXT("spvasm"));
+		DumpDebugShaderBinary(InternalState.Input, SerializedOutput.Spirv.GetByteData(), SerializedOutput.Spirv.GetByteSize(), SPVExt);
+		DumpDebugShaderDisassembledSpirv(InternalState.Input, SerializedOutput.Spirv.GetByteData(), SerializedOutput.Spirv.GetByteSize(), SPVASMExt);
 	}
 
 	return true;
@@ -2083,7 +2161,12 @@ bool PreprocessVulkanShader(const FShaderCompilerInput& Input, const FShaderComp
 		}
 	}
 
-	if (!PreprocessOutput.ParseAndModify(Input, Environment, EBindlessParameterMode::Vulkan))
+	// Create a _RootShaderParameters and bind it in slot 0 like any other uniform buffer
+	const TCHAR* ConstantBufferType = InternalState.UseRootParametersStructure() ? TEXT("cbuffer") : nullptr;
+
+	const TArrayView<const TCHAR* const> ExtraSRVTypes;
+	const TArrayView<const TCHAR* const> ExtraUAVTypes;
+	if (!PreprocessOutput.ParseAndModify(Input, Environment, ConstantBufferType, ExtraSRVTypes, ExtraUAVTypes, EBindlessParameterMode::Vulkan))
 	{
 		// The FShaderParameterParser will add any relevant errors.
 		return false;
@@ -2095,8 +2178,13 @@ bool PreprocessVulkanShader(const FShaderCompilerInput& Input, const FShaderComp
 	TransformStringIntoCharacterArray(PreprocessedShaderSource);
 
 	// Run the shader minifier
+	if (InternalState.IsRayTracingShader() || InternalState.bUseBindlessUniformBuffer)
+	{
+		// Always needed in ray tracing shaders to ensure bindless UB count is kept low
+		UE::ShaderCompilerCommon::RemoveDeadCode(PreprocessedShaderSource, InternalState.GetRequiredSymbols(), PreprocessOutput.EditErrors());
+	}
 	#if UE_VULKAN_SHADER_COMPILER_ALLOW_DEAD_CODE_REMOVAL
-	if (Input.Environment.CompilerFlags.Contains(CFLAG_RemoveDeadCode))
+	else if (Input.Environment.CompilerFlags.Contains(CFLAG_RemoveDeadCode))
 	{
 		UE::ShaderCompilerCommon::RemoveDeadCode(PreprocessedShaderSource, InternalState.GetEntryPointName(), PreprocessOutput.EditErrors());
 	}
