VulkanRHI: VK RT WIP: Shader compiler changes for Vulkan RT.  Add an option for extra blobs for AnyHit and Intersection shaders in HitGroup.  Add option of bindless UBs for ray tracing shaders.  Correctly set GRHISupportsRayTracingDispatchIndirect.

#rnx

diff --git a/Engine/Shaders/Public/Platform/Vulkan/VulkanCommon.ush b/Engine/Shaders/Public/Platform/Vulkan/VulkanCommon.ush
--- a/Engine/Shaders/Public/Platform/Vulkan/VulkanCommon.ush
+++ b/Engine/Shaders/Public/Platform/Vulkan/VulkanCommon.ush
@@ -5,7 +5,7 @@
 =============================================================================*/
 
 // Update this GUID to improve shader recompilation for Vulkan only shaders
-// GUID = 77A59EF7-19F3-4EE5-9058-F1E42AAAE164
+// GUID = EBB16B3A-3FC4-4518-BDF0-63FCF6281A9A
 
 #pragma once
 
diff --git a/Engine/Source/Developer/ShaderCompilerCommon/Public/ShaderCompilerCommon.h b/Engine/Source/Developer/ShaderCompilerCommon/Public/ShaderCompilerCommon.h
--- a/Engine/Source/Developer/ShaderCompilerCommon/Public/ShaderCompilerCommon.h
+++ b/Engine/Source/Developer/ShaderCompilerCommon/Public/ShaderCompilerCommon.h
@@ -310,6 +310,9 @@ extern SHADERCOMPILERCOMMON_API void AddUnboundShaderParameterError(
 // The cross compiler doesn't yet support struct initializers needed to construct static structs for uniform buffers
 // Replace all uniform buffer struct member references (View.WorldToClip) with a flattened name that removes the struct dependency (View_WorldToClip)
 extern SHADERCOMPILERCOMMON_API void RemoveUniformBuffersFromSource(const FShaderCompilerEnvironment& Environment, FString& PreprocessedShaderSource);
+extern SHADERCOMPILERCOMMON_API const TCHAR* FindMatchingClosingBrace(const TCHAR* OpeningCharPtr);
+extern SHADERCOMPILERCOMMON_API const TCHAR* ParseHLSLSymbolName(const TCHAR* SearchString, FString& SymboName);
+extern SHADERCOMPILERCOMMON_API void ParseHLSLTypeName(const TCHAR* SearchString, const TCHAR*& TypeNameStartPtr, const TCHAR*& TypeNameEndPtr);
 
 // Processes TEXT macros
 extern SHADERCOMPILERCOMMON_API void TransformStringIntoCharacterArray(FString& PreprocessedShaderSource, TArray<FShaderDiagnosticData>* OutDiagnosticDatas=nullptr);
diff --git a/Engine/Source/Developer/VulkanShaderFormat/Private/VulkanShaderCompiler.cpp b/Engine/Source/Developer/VulkanShaderFormat/Private/VulkanShaderCompiler.cpp
--- a/Engine/Source/Developer/VulkanShaderFormat/Private/VulkanShaderCompiler.cpp
+++ b/Engine/Source/Developer/VulkanShaderFormat/Private/VulkanShaderCompiler.cpp
@@ -1139,16 +1139,6 @@ static void BuildShaderOutput(
 	// Build the SRT for this shader from the NEWHeader
 	SerializedOutput.ShaderResourceTable = BuildSRTFromHeader(NEWHeader);
 
-	// Write out the header and shader source code.
-	FMemoryWriter Ar(ShaderOutput.ShaderCode.GetWriteAccess(), true);
-	Ar << SerializedOutput.Header;
-	Ar << SerializedOutput.ShaderResourceTable;
-
-	check(SerializedOutput.Spirv.Data.Num() != 0);
-	uint32 SpirvCodeSizeBytes = SerializedOutput.Spirv.GetByteSize();
-	Ar << SpirvCodeSizeBytes;
-	Ar.Serialize((uint8*)SerializedOutput.Spirv.Data.GetData(), SpirvCodeSizeBytes);
-
 	ShaderOutput.bSucceeded = true;
 
 	if (ShaderInput.ExtraSettings.bExtractShaderSource)
@@ -1577,7 +1567,15 @@ static bool BuildShaderOutputFromSpirv(
 
 	PatchSpirvReflectionEntries(SerializedOutput.Spirv);
 
-	SerializedOutput.Spirv.EntryPointName = PatchSpirvEntryPointWithCRC(SerializedOutput.Spirv, SerializedOutput.Spirv.CRC);
+	// :todo-jn: We don't store the CRC of each member of the hit group, leave the entrypoint untouched on the extra modules
+	if (InternalState.HasMultipleEntryPoints() && (InternalState.HitGroupShaderType != FVulkanShaderCompilerInternalState::EHitGroupShaderType::ClosestHit))
+	{
+		SerializedOutput.Spirv.EntryPointName = "main_00000000_00000000";
+	}
+	else
+	{
+		SerializedOutput.Spirv.EntryPointName = PatchSpirvEntryPointWithCRC(SerializedOutput.Spirv, SerializedOutput.Spirv.CRC);
+	}
 
 	Output.NumInstructions = CalculateSpirvInstructionCount(SerializedOutput.Spirv);
 
@@ -2075,6 +2073,54 @@ static bool CompileWithShaderConductor(
 
 
 
+static void RemoveUnusedBindlessHeaps(FString& PreprocessedShaderSource, const TCHAR* HeapType)
+{
+	const FString HeapIdentifier = TEXT("VULKAN_") + FString(HeapType) + TEXT("_HEAP(");
+
+	int32 SearchIndex = PreprocessedShaderSource.Find(HeapIdentifier, ESearchCase::CaseSensitive, ESearchDir::FromStart, -1);
+	while (SearchIndex != INDEX_NONE)
+	{
+		const int32 TypeNameStartIndex = SearchIndex + HeapIdentifier.Len();
+		const int32 TypeNameEndIndex = PreprocessedShaderSource.Find(TEXT(")"), ESearchCase::CaseSensitive, ESearchDir::FromStart, TypeNameStartIndex);
+		FString TypeName(TypeNameEndIndex - TypeNameStartIndex, &PreprocessedShaderSource[TypeNameStartIndex]);
+		TypeName.TrimStartAndEndInline();
+
+		// Make sure it's one of our automatically generated types
+		if (TypeName.StartsWith(TEXT("SafeType")))
+		{
+			// Ugly and fast way to make sure ity's a heap declaration (this should catch the generated ones at least)
+			if ((PreprocessedShaderSource[TypeNameEndIndex + 1] == '[') &&
+				(PreprocessedShaderSource[TypeNameEndIndex + 2] == ']') &&
+				(PreprocessedShaderSource[TypeNameEndIndex + 3] == ';'))
+			{
+				int32 FirstUseIndex = PreprocessedShaderSource.Find(TypeName, ESearchCase::CaseSensitive, ESearchDir::FromStart, TypeNameEndIndex + 4);
+				while (FirstUseIndex != INDEX_NONE)
+				{
+					const int32 NextChar = FirstUseIndex + TypeName.Len();
+					if (FChar::IsWhitespace(PreprocessedShaderSource[NextChar]) || PreprocessedShaderSource[NextChar] == ')')
+					{
+						break;
+					}
+					FirstUseIndex = PreprocessedShaderSource.Find(TypeName, ESearchCase::CaseSensitive, ESearchDir::FromStart, NextChar);
+				}
+
+				if (FirstUseIndex == INDEX_NONE)
+				{
+					const int32 DeclarationBeginIndex = PreprocessedShaderSource.Find(TypeName, ESearchCase::CaseSensitive, ESearchDir::FromEnd, SearchIndex);
+					if ((DeclarationBeginIndex != INDEX_NONE) && ((SearchIndex - DeclarationBeginIndex) < (TypeName.Len() + 4)))
+					{
+						for (int32 Idx = DeclarationBeginIndex; Idx <= (TypeNameEndIndex + 3); Idx++)
+						{
+							PreprocessedShaderSource[Idx] = ' ';
+						}
+					}
+				}
+			}
+		}
+
+		SearchIndex = PreprocessedShaderSource.Find(HeapIdentifier, ESearchCase::CaseSensitive, ESearchDir::FromStart, TypeNameEndIndex + 4);
+	}
+}
 
 
 
@@ -2190,14 +2236,276 @@ bool PreprocessVulkanShader(const FShaderCompilerInput& Input, const FShaderComp
 	}
 	#endif // UE_VULKAN_SHADER_COMPILER_ALLOW_DEAD_CODE_REMOVAL
 
+	// Clean up the code a bit, it's unreadable otherwise with all the unused heaps left around
+	if (InternalState.bSupportsBindless)
+	{
+		RemoveUnusedBindlessHeaps(PreprocessedShaderSource, TEXT("SAMPLER"));
+		RemoveUnusedBindlessHeaps(PreprocessedShaderSource, TEXT("RESOURCE"));
+		UE::ShaderCompilerCommon::RemoveDeadCode(PreprocessedShaderSource, InternalState.GetRequiredSymbols(), PreprocessOutput.EditErrors());
+
+		if (InternalState.bDebugDump)
+		{
+			DumpDebugShaderText(Input, PreprocessedShaderSource, TEXT("bindless.final.hlsl"));
+		}
+	}
+
 	return true;
 }
 
+
+// :todo-jn: TEMPORARY EXPERIMENT - will eventually move into preprocessing step
+static TArray<FString> ConvertUBToBindless(FString& PreprocessedShaderSource)
+{
+	// Fill a map so we pull our bindless sampler/resource indices from the right struct
+	// :todo-jn: Do we not have the layout somewhere instead of calculating offsets?  there must be a better way...
+	auto GenerateNewDecl = [](const int32 CBIndex, const FString& Members, const FString& CBName)
+	{
+		const FString PrefixedCBName = FString::Printf(TEXT("%s%d_%s"), *kBindlessCBPrefix, CBIndex, *CBName);
+		const FString BindlessCBType = PrefixedCBName + TEXT("_Type");
+		const FString BindlessCBHeapName = PrefixedCBName + kBindlessHeapSuffix;
+
+		FString CBDecl;
+		CBDecl.Reserve(Members.Len() * 3);  // start somewhere approx less bad
+
+		// Declare the struct
+		CBDecl += TEXT("struct ") + BindlessCBType + TEXT(" \n{\n") + Members + TEXT("\n};\n");
+
+		// Declare the safetype and bindless array for this cb
+		CBDecl += FString::Printf(TEXT("ConstantBuffer<%s> %s[];\n"), *BindlessCBType, *BindlessCBHeapName);
+
+		// Now bring in the CB
+		CBDecl += FString::Printf(TEXT("static const %s %s = %s[VulkanHitGroupSystemParameters.BindlessUniformBuffers[%d]];\n"),
+			*BindlessCBType, *PrefixedCBName, *BindlessCBHeapName, CBIndex);
+
+		// Now create a global scope var for each value (as the cbuffer would provide) to patch in seemlessly with the rest of the code
+		uint32 MemberOffset = 0;
+		const TCHAR* MemberSearchPtr = *Members;
+		const uint32 LastMemberSemicolonIndex = Members.Find(TEXT(";"), ESearchCase::CaseSensitive, ESearchDir::FromEnd, -1);
+		check(LastMemberSemicolonIndex != INDEX_NONE);
+		const TCHAR* LastMemberSemicolon = &Members[LastMemberSemicolonIndex];
+
+		do
+		{
+			const TCHAR* MemberTypeStartPtr = nullptr;
+			const TCHAR* MemberTypeEndPtr = nullptr;
+			ParseHLSLTypeName(MemberSearchPtr, MemberTypeStartPtr, MemberTypeEndPtr);
+			const FString MemberTypeName(MemberTypeEndPtr - MemberTypeStartPtr, MemberTypeStartPtr);
+
+			FString MemberName;
+			MemberSearchPtr = ParseHLSLSymbolName(MemberTypeEndPtr, MemberName);
+			check(MemberName.Len() > 0);
+
+			// Skip over trailing tokens and pick up arrays
+			FString ArrayDecl;
+			while (*MemberSearchPtr && *MemberSearchPtr != ';')
+			{
+				if (*MemberSearchPtr == '[')
+				{
+					ArrayDecl.AppendChar(*MemberSearchPtr);
+
+					MemberSearchPtr++;
+					while (*MemberSearchPtr && *MemberSearchPtr != ']')
+					{
+						ArrayDecl.AppendChar(*MemberSearchPtr);
+						MemberSearchPtr++;
+					}
+
+					ArrayDecl.AppendChar(*MemberSearchPtr);
+				}
+
+				MemberSearchPtr++;
+			}
+
+			CBDecl += FString::Printf(TEXT("static const %s %s%s = %s.%s;\n"), *MemberTypeName, *MemberName, *ArrayDecl, *PrefixedCBName, *MemberName);
+
+		} while (MemberSearchPtr < LastMemberSemicolon);
+
+		return CBDecl;
+	};
+
+	// replace "cbuffer" decl with a struct filled from bindless constant buffer
+	TArray<FString> BindlessUBs;
+	{
+		const FString UniformBufferDeclIdentifier = TEXT("cbuffer");
+
+		int32 SearchIndex = PreprocessedShaderSource.Find(UniformBufferDeclIdentifier, ESearchCase::CaseSensitive, ESearchDir::FromStart, -1);
+		while (SearchIndex != INDEX_NONE)
+		{
+			FString StructName;
+			const TCHAR* StructNameEndPtr = ParseHLSLSymbolName(&PreprocessedShaderSource[SearchIndex + UniformBufferDeclIdentifier.Len()], StructName);
+			check(StructName.Len() > 0);
+
+			const int32 CBIndex = BindlessUBs.Add(StructName);
+			check(CBIndex < 16);
+
+			const TCHAR* OpeningBracePtr = FCString::Strstr(&PreprocessedShaderSource[SearchIndex + UniformBufferDeclIdentifier.Len()], TEXT("{"));
+			check(OpeningBracePtr);
+			const TCHAR* ClosingBracePtr = FindMatchingClosingBrace(OpeningBracePtr + 1);
+			check(ClosingBracePtr);
+			const int32 ClosingBraceIndex = ClosingBracePtr - (*PreprocessedShaderSource);
+
+			const FString Members(ClosingBracePtr - OpeningBracePtr - 1, OpeningBracePtr + 1);
+			const FString NewDecl = GenerateNewDecl(CBIndex, Members, StructName);
+
+			const int32 OldDeclLen = ClosingBraceIndex - SearchIndex + 1;
+			PreprocessedShaderSource.RemoveAt(SearchIndex, OldDeclLen, false);
+			PreprocessedShaderSource.InsertAt(SearchIndex, NewDecl);
+
+			SearchIndex = PreprocessedShaderSource.Find(UniformBufferDeclIdentifier, ESearchCase::CaseSensitive, ESearchDir::FromStart, SearchIndex + NewDecl.Len());
+		}
+	}
+	return BindlessUBs;
+}
+
+
+static bool CompileShaderGroup(
+	FVulkanShaderCompilerInternalState& InternalState,
+	const FString& OriginalPreprocessedShaderSource,
+	FShaderCompilerOutput& MergedOutput
+)
+{
+	checkf(InternalState.bSupportsBindless && InternalState.bUseBindlessUniformBuffer, TEXT("Ray tracing requires full bindless in Vulkan."));
+
+	// Compile each one of the shader modules seperately and create one big blob for the engine
+	auto CompilePartialExport = [&OriginalPreprocessedShaderSource, &InternalState, &MergedOutput](
+		FVulkanShaderCompilerInternalState::EHitGroupShaderType HitGroupShaderType,
+		const TCHAR* PartialFileExtension,
+		VulkanShaderCompilerSerializedOutput& PartialSerializedOutput)
+	{
+		InternalState.HitGroupShaderType = HitGroupShaderType;
+
+		FShaderCompilerOutput TempOutput;
+		const bool bIsClosestHit = (HitGroupShaderType == FVulkanShaderCompilerInternalState::EHitGroupShaderType::ClosestHit);
+		FShaderCompilerOutput& PartialOutput = bIsClosestHit ? MergedOutput : TempOutput;
+
+		FString PartialPreprocessedShaderSource = OriginalPreprocessedShaderSource;
+		UE::ShaderCompilerCommon::RemoveDeadCode(PartialPreprocessedShaderSource, InternalState.GetEntryPointName(), PartialOutput.Errors);
+
+		if (InternalState.bDebugDump)
+		{
+			DumpDebugShaderText(InternalState.Input, PartialPreprocessedShaderSource, *FString::Printf(TEXT("%s.hlsl"), PartialFileExtension));
+		}
+
+		const bool bPartialSuccess = CompileWithShaderConductor(InternalState, PartialPreprocessedShaderSource, PartialSerializedOutput, PartialOutput);
+
+		if (!bIsClosestHit)
+		{
+			MergedOutput.NumInstructions = FMath::Max(MergedOutput.NumInstructions, PartialOutput.NumInstructions);
+			MergedOutput.NumTextureSamplers = FMath::Max(MergedOutput.NumTextureSamplers, PartialOutput.NumTextureSamplers);
+			MergedOutput.Errors.Append(MoveTemp(PartialOutput.Errors));
+		}
+
+		return bPartialSuccess;
+	};
+
+	bool bSuccess = false;
+
+	// Closest Hit Module, always present
+	VulkanShaderCompilerSerializedOutput ClosestHitSerializedOutput;
+	{
+		bSuccess = CompilePartialExport(FVulkanShaderCompilerInternalState::EHitGroupShaderType::ClosestHit, TEXT("closest"), ClosestHitSerializedOutput);
+	}
+
+	// Any Hit Module, optional
+	const bool bHasAnyHitModule = !InternalState.AnyHitEntry.IsEmpty();
+	VulkanShaderCompilerSerializedOutput AnyHitSerializedOutput;
+	if (bSuccess && bHasAnyHitModule)
+	{
+		bSuccess = CompilePartialExport(FVulkanShaderCompilerInternalState::EHitGroupShaderType::AnyHit, TEXT("anyhit"), AnyHitSerializedOutput);
+	}
+
+	// Intersection Module, optional
+	const bool bHasIntersectionModule = !InternalState.IntersectionEntry.IsEmpty();
+	VulkanShaderCompilerSerializedOutput IntersectionSerializedOutput;
+	if (bSuccess && bHasIntersectionModule)
+	{
+		bSuccess = CompilePartialExport(FVulkanShaderCompilerInternalState::EHitGroupShaderType::Intersection, TEXT("intersection"), IntersectionSerializedOutput);
+	}
+
+	// Add the bindless data in the output at the very end
+	{
+		auto GetLayoutHash = [&InternalState](const FString& UBName)
+		{
+			uint32 LayoutHash = 0;
+			const FUniformBufferEntry* UniformBufferEntry = InternalState.Input.Environment.UniformBufferMap.Find(UBName);
+			if (UniformBufferEntry)
+			{
+				LayoutHash = UniformBufferEntry->LayoutHash;
+			}
+			else if ((UBName == FShaderParametersMetadata::kRootUniformBufferBindingName) && InternalState.Input.RootParametersStructure)
+			{
+				LayoutHash = InternalState.Input.RootParametersStructure->GetLayoutHash();
+			}
+			else
+			{
+				LayoutHash = 0;
+			}
+			return LayoutHash;
+		};
+
+		ClosestHitSerializedOutput.Header.UniformBuffers.Empty();
+		for (int32 CBIndex = 0; CBIndex < InternalState.AllBindlessUBs.Num(); CBIndex++)
+		{
+			const FString& CBName = InternalState.AllBindlessUBs[CBIndex];
+
+			// It's possible SPIRV compilation has optimized out a buffer from every shader in the group
+			if (ClosestHitSerializedOutput.UsedBindlessUB.Contains(CBName) ||
+				AnyHitSerializedOutput.UsedBindlessUB.Contains(CBName) ||
+				IntersectionSerializedOutput.UsedBindlessUB.Contains(CBName))
+			{
+				FVulkanShaderHeader::FUniformBufferInfo& UBInfo = ClosestHitSerializedOutput.Header.UniformBuffers.AddZeroed_GetRef();
+				UBInfo.LayoutHash = GetLayoutHash(CBName);
+				UBInfo.ConstantDataOriginalBindingIndex = CBIndex;
+#if VULKAN_ENABLE_BINDING_DEBUG_NAMES
+				UBInfo.DebugName = CBName;
+#endif
+
+				MergedOutput.ParameterMap.AddParameterAllocation(*CBName, CBIndex, (uint16)FVulkanShaderHeader::UniformBuffer, 1, EShaderParameterType::UniformBuffer);
+			}
+		}
+	}
+
+	{
+		// :todo-jn: Having multiple entrypoints in a single SPIRV blob crashes on FLumenHardwareRayTracingMaterialHitGroup for some reason
+		// Adjust the header before we write it out
+		ClosestHitSerializedOutput.Header.RayGroupAnyHit = bHasAnyHitModule ? FVulkanShaderHeader::ERayHitGroupEntrypoint::SeparateBlob : FVulkanShaderHeader::ERayHitGroupEntrypoint::NotPresent;
+		ClosestHitSerializedOutput.Header.RayGroupIntersection = bHasIntersectionModule ? FVulkanShaderHeader::ERayHitGroupEntrypoint::SeparateBlob : FVulkanShaderHeader::ERayHitGroupEntrypoint::NotPresent;
+
+		check(ClosestHitSerializedOutput.Spirv.Data.Num() != 0);
+		FMemoryWriter Ar(MergedOutput.ShaderCode.GetWriteAccess(), true);
+		Ar << ClosestHitSerializedOutput.Header;
+		Ar << ClosestHitSerializedOutput.ShaderResourceTable;
+
+		{
+			uint32 SpirvCodeSizeBytes = ClosestHitSerializedOutput.Spirv.GetByteSize();
+			Ar << SpirvCodeSizeBytes;
+			Ar.Serialize((uint8*)ClosestHitSerializedOutput.Spirv.Data.GetData(), SpirvCodeSizeBytes);
+		}
+
+		if (bHasAnyHitModule)
+		{
+			uint32 SpirvCodeSizeBytes = AnyHitSerializedOutput.Spirv.GetByteSize();
+			Ar << SpirvCodeSizeBytes;
+			Ar.Serialize((uint8*)AnyHitSerializedOutput.Spirv.Data.GetData(), SpirvCodeSizeBytes);
+		}
+
+		if (bHasIntersectionModule)
+		{
+			uint32 SpirvCodeSizeBytes = IntersectionSerializedOutput.Spirv.GetByteSize();
+			Ar << SpirvCodeSizeBytes;
+			Ar.Serialize((uint8*)IntersectionSerializedOutput.Spirv.Data.GetData(), SpirvCodeSizeBytes);
+		}
+	}
+
+	MergedOutput.bSucceeded = bSuccess;
+	return bSuccess;
+}
+
 void CompileVulkanShader(const FShaderCompilerInput& Input, const FShaderPreprocessOutput& PreprocessOutput, FShaderCompilerOutput& Output, const class FString& WorkingDirectory)
 {
 	check(IsVulkanShaderFormat(Input.ShaderFormat));
 
-	const FVulkanShaderCompilerInternalState InternalState(Input);
+	FVulkanShaderCompilerInternalState InternalState(Input);
 
 	const EHlslShaderFrequency HlslFrequency = InternalState.GetHlslShaderFrequency();
 	if (HlslFrequency == HSF_InvalidFrequency)
@@ -2210,12 +2518,38 @@ void CompileVulkanShader(const FShaderCompilerInput& Input, const FShaderPreproc
 		return;
 	}
 
+	FString BindlessPreprocessedShaderSource;
+	if (InternalState.bUseBindlessUniformBuffer)
+	{
+		BindlessPreprocessedShaderSource = PreprocessOutput.GetSource();
+		InternalState.AllBindlessUBs = ConvertUBToBindless(BindlessPreprocessedShaderSource);
+	}
+	const FString& PreprocessedShaderSource = InternalState.bUseBindlessUniformBuffer ? BindlessPreprocessedShaderSource : PreprocessOutput.GetSource();
+
 	bool bSuccess = false;
 
 #if PLATFORM_MAC || PLATFORM_WINDOWS || PLATFORM_LINUX
-	// Compile shader via ShaderConductor (DXC, SPIRV-Tools)
-	VulkanShaderCompilerSerializedOutput SerializedOutput;
-	bSuccess = CompileWithShaderConductor(InternalState, PreprocessOutput.GetSource(), SerializedOutput, Output);
+	// HitGroup shaders might have multiple entrypoints that we combine into a single blob
+	if (InternalState.HasMultipleEntryPoints())
+	{
+		bSuccess = CompileShaderGroup(InternalState, PreprocessedShaderSource, Output);
+	}
+	else
+	{
+		// Compile regular shader via ShaderConductor (DXC)
+		VulkanShaderCompilerSerializedOutput SerializedOutput;
+		bSuccess = CompileWithShaderConductor(InternalState, PreprocessedShaderSource, SerializedOutput, Output);
+
+		// Write out the header and shader source code (except for the extra shaders in hit groups)
+		check(SerializedOutput.Spirv.Data.Num() != 0);
+		FMemoryWriter Ar(Output.ShaderCode.GetWriteAccess(), true);
+		Ar << SerializedOutput.Header;
+		Ar << SerializedOutput.ShaderResourceTable;
+
+		uint32 SpirvCodeSizeBytes = SerializedOutput.Spirv.GetByteSize();
+		Ar << SpirvCodeSizeBytes;
+		Ar.Serialize((uint8*)SerializedOutput.Spirv.Data.GetData(), SpirvCodeSizeBytes);
+	}
 #endif // PLATFORM_MAC || PLATFORM_WINDOWS || PLATFORM_LINUX
 	
 	if (Input.Environment.CompilerFlags.Contains(CFLAG_ExtraShaderData))
diff --git a/Engine/Source/Runtime/VulkanRHI/Private/VulkanExtensions.cpp b/Engine/Source/Runtime/VulkanRHI/Private/VulkanExtensions.cpp
--- a/Engine/Source/Runtime/VulkanRHI/Private/VulkanExtensions.cpp
+++ b/Engine/Source/Runtime/VulkanRHI/Private/VulkanExtensions.cpp
@@ -830,6 +830,8 @@ public:
 		if (bRequirementsPassed)
 		{
 			AddToPNext(DeviceCreateInfo, RayTracingPipelineFeatures);
+
+			GRHISupportsRayTracingDispatchIndirect = (RayTracingPipelineFeatures.rayTracingPipelineTraceRaysIndirect == VK_TRUE);
 		}
 	}
 
diff --git a/Engine/Source/Runtime/VulkanRHI/Private/VulkanRayTracing.cpp b/Engine/Source/Runtime/VulkanRHI/Private/VulkanRayTracing.cpp
--- a/Engine/Source/Runtime/VulkanRHI/Private/VulkanRayTracing.cpp
+++ b/Engine/Source/Runtime/VulkanRHI/Private/VulkanRayTracing.cpp
@@ -931,6 +931,11 @@ void FVulkanRayTracingShaderTable::Init(const FVulkanRayTracingScene* Scene, con
 	InitAlloc(Miss, SceneInitializer.NumMissShaderSlots, true);
 	InitAlloc(HitGroup, Pipeline->bAllowHitGroupIndexing ? SceneInitializer.NumTotalSegments * SceneInitializer.ShaderSlotsPerGeometrySegment : 1, true);
 	InitAlloc(Callable, SceneInitializer.NumCallableShaderSlots, true);
+
+	if (!Pipeline->bAllowHitGroupIndexing && Pipeline->GetShaderHandles(SF_RayHitGroup).Num())
+	{
+		SetSlot(SF_RayHitGroup, 0, 0, Pipeline->GetShaderHandles(SF_RayHitGroup));
+	}
 }
 
 FVulkanRayTracingShaderTable::FVulkanShaderTableAllocation& FVulkanRayTracingShaderTable::GetAlloc(EShaderFrequency Frequency)
@@ -1317,7 +1322,7 @@ FVulkanRayTracingPipelineState::FVulkanRayTracingPipelineState(FVulkanDevice* co
 
 		VkPipelineShaderStageCreateInfo ShaderStage;
 		ZeroVulkanStruct(ShaderStage, VK_STRUCTURE_TYPE_PIPELINE_SHADER_STAGE_CREATE_INFO);
-		ShaderStage.module = RayGenShader->GetOrCreateHandle()->GetVkShaderModule();
+		ShaderStage.module = RayGenShader->GetOrCreateHandle(FVulkanRayTracingShader::MainModuleIdentifier)->GetVkShaderModule();
 		ShaderStage.stage = VK_SHADER_STAGE_RAYGEN_BIT_KHR;
 			
 		ANSICHAR* const EntryPoint = new ANSICHAR[EntryPointNameMaxLength];
@@ -1346,7 +1351,7 @@ FVulkanRayTracingPipelineState::FVulkanRayTracingPipelineState(FVulkanDevice* co
 
 		VkPipelineShaderStageCreateInfo ShaderStage;
 		ZeroVulkanStruct(ShaderStage, VK_STRUCTURE_TYPE_PIPELINE_SHADER_STAGE_CREATE_INFO);
-		ShaderStage.module = MissShader->GetOrCreateHandle()->GetVkShaderModule();
+		ShaderStage.module = MissShader->GetOrCreateHandle(FVulkanRayTracingShader::MainModuleIdentifier)->GetVkShaderModule();
 		ShaderStage.stage = VK_SHADER_STAGE_MISS_BIT_KHR;
 
 		ANSICHAR* const EntryPoint = new char[EntryPointNameMaxLength];
@@ -1370,27 +1375,58 @@ FVulkanRayTracingPipelineState::FVulkanRayTracingPipelineState(FVulkanDevice* co
 	HitGroup.Shaders.Reserve(InitializerHitGroupShaders.Num());
 	for (FRHIRayTracingShader* const HitGroupShaderRHI : InitializerHitGroupShaders)
 	{
+		VkRayTracingShaderGroupCreateInfoKHR ShaderGroup;
+		ZeroVulkanStruct(ShaderGroup, VK_STRUCTURE_TYPE_RAY_TRACING_SHADER_GROUP_CREATE_INFO_KHR);
+		ShaderGroup.type = VK_RAY_TRACING_SHADER_GROUP_TYPE_TRIANGLES_HIT_GROUP_KHR;
+		ShaderGroup.generalShader = VK_SHADER_UNUSED_KHR;
+
 		checkSlow(HitGroupShaderRHI->GetFrequency() == SF_RayHitGroup);
 		FVulkanRayTracingShader* const HitGroupShader = ResourceCast(HitGroupShaderRHI);
 
-		VkPipelineShaderStageCreateInfo ShaderStage;
-		ZeroVulkanStruct(ShaderStage, VK_STRUCTURE_TYPE_PIPELINE_SHADER_STAGE_CREATE_INFO);
-		ShaderStage.module = HitGroupShader->GetOrCreateHandle()->GetVkShaderModule();
-		ShaderStage.stage = VK_SHADER_STAGE_CLOSEST_HIT_BIT_KHR;
+		// Closest Hit, always present
+		{
+			ANSICHAR* const EntryPoint = new char[EntryPointNameMaxLength];
+			HitGroupShader->GetEntryPoint(EntryPoint, EntryPointNameMaxLength);
+			EntryPointNames.Add(EntryPoint);
+
+			VkPipelineShaderStageCreateInfo ShaderStage;
+			ZeroVulkanStruct(ShaderStage, VK_STRUCTURE_TYPE_PIPELINE_SHADER_STAGE_CREATE_INFO);
+			ShaderStage.module = HitGroupShader->GetOrCreateHandle(FVulkanRayTracingShader::ClosestHitModuleIdentifier)->GetVkShaderModule();
+			ShaderStage.stage = VK_SHADER_STAGE_CLOSEST_HIT_BIT_KHR;
+			ShaderStage.pName = EntryPoint;
+			ShaderGroup.closestHitShader = ShaderStages.Add(ShaderStage);
+		}
 
-		ANSICHAR* const EntryPoint = new char[EntryPointNameMaxLength];
-		HitGroupShader->GetEntryPoint(EntryPoint, EntryPointNameMaxLength);
-		EntryPointNames.Add(EntryPoint);
-		ShaderStage.pName = EntryPoint;
-		ShaderStages.Add(ShaderStage);
+		// Any Hit, optional
+		if (HitGroupShader->GetCodeHeader().RayGroupAnyHit != FVulkanShaderHeader::ERayHitGroupEntrypoint::NotPresent)
+		{
+			VkPipelineShaderStageCreateInfo ShaderStage;
+			ZeroVulkanStruct(ShaderStage, VK_STRUCTURE_TYPE_PIPELINE_SHADER_STAGE_CREATE_INFO);
+			ShaderStage.module = HitGroupShader->GetOrCreateHandle(FVulkanRayTracingShader::AnyHitModuleIdentifier)->GetVkShaderModule();
+			ShaderStage.stage = VK_SHADER_STAGE_ANY_HIT_BIT_KHR;
+			ShaderStage.pName = "main_00000000_00000000"; // :todo-jn: patch in the size_crc
+			ShaderGroup.anyHitShader = ShaderStages.Add(ShaderStage);
+		}
+		else
+		{
+			ShaderGroup.anyHitShader = VK_SHADER_UNUSED_KHR;
+		}
+
+		// Intersection, optional
+		if (HitGroupShader->GetCodeHeader().RayGroupIntersection != FVulkanShaderHeader::ERayHitGroupEntrypoint::NotPresent)
+		{
+			VkPipelineShaderStageCreateInfo ShaderStage;
+			ZeroVulkanStruct(ShaderStage, VK_STRUCTURE_TYPE_PIPELINE_SHADER_STAGE_CREATE_INFO);
+			ShaderStage.module = HitGroupShader->GetOrCreateHandle(FVulkanRayTracingShader::IntersectionModuleIdentifier)->GetVkShaderModule();
+			ShaderStage.stage = VK_SHADER_STAGE_INTERSECTION_BIT_KHR;
+			ShaderStage.pName = "main_00000000_00000000"; // :todo-jn: patch in the size_crc
+			ShaderGroup.intersectionShader = ShaderStages.Add(ShaderStage);
+		}
+		else
+		{
+			ShaderGroup.intersectionShader = VK_SHADER_UNUSED_KHR;
+		}
 
-		VkRayTracingShaderGroupCreateInfoKHR ShaderGroup;
-		ZeroVulkanStruct(ShaderGroup, VK_STRUCTURE_TYPE_RAY_TRACING_SHADER_GROUP_CREATE_INFO_KHR);
-		ShaderGroup.type = VK_RAY_TRACING_SHADER_GROUP_TYPE_TRIANGLES_HIT_GROUP_KHR;
-		ShaderGroup.generalShader = VK_SHADER_UNUSED_KHR;
-		ShaderGroup.closestHitShader = ShaderStages.Num() - 1;
-		ShaderGroup.anyHitShader = VK_SHADER_UNUSED_KHR; // vkrt: todo
-		ShaderGroup.intersectionShader = VK_SHADER_UNUSED_KHR;
 		ShaderGroups.Add(ShaderGroup);
 
 		HitGroup.Shaders.Add(HitGroupShader);
@@ -1404,7 +1440,7 @@ FVulkanRayTracingPipelineState::FVulkanRayTracingPipelineState(FVulkanDevice* co
 
 		VkPipelineShaderStageCreateInfo ShaderStage;
 		ZeroVulkanStruct(ShaderStage, VK_STRUCTURE_TYPE_PIPELINE_SHADER_STAGE_CREATE_INFO);
-		ShaderStage.module = CallableShader->GetOrCreateHandle()->GetVkShaderModule();
+		ShaderStage.module = CallableShader->GetOrCreateHandle(FVulkanRayTracingShader::MainModuleIdentifier)->GetVkShaderModule();
 		ShaderStage.stage = VK_SHADER_STAGE_CALLABLE_BIT_KHR;
 
 		ANSICHAR* const EntryPoint = new char[EntryPointNameMaxLength];
diff --git a/Engine/Source/Runtime/VulkanRHI/Private/VulkanResources.h b/Engine/Source/Runtime/VulkanRHI/Private/VulkanResources.h
--- a/Engine/Source/Runtime/VulkanRHI/Private/VulkanResources.h
+++ b/Engine/Source/Runtime/VulkanRHI/Private/VulkanResources.h
@@ -99,6 +99,8 @@ public:
 
 class FVulkanShader : public IRefCountedObject
 {
+protected:
+
 	static FCriticalSection VulkanShaderModulesMapCS;
 
 public:
@@ -180,9 +182,15 @@ public:
 	public:
 		TArrayView<uint32> GetCodeView() {return CodeView;}
 	};
-	FSpirvCode GetSpirvCode();
+
+	inline FSpirvCode GetSpirvCode()
+	{
+		return GetSpirvCode(SpirvContainer);
+	}
+	
 	FSpirvCode GetPatchedSpirvCode(const FGfxPipelineDesc& Desc, const FVulkanLayout* Layout);
 protected:
+
 #if UE_BUILD_DEBUG || UE_BUILD_DEVELOPMENT
 	FString							DebugEntryPoint;
 #endif
@@ -197,7 +205,7 @@ protected:
 
 	FShaderResourceTable			ShaderResourceTable;
 
-private:
+protected:
 	class FSpirvContainer
 	{
 		friend class FVulkanShader;
@@ -212,6 +220,8 @@ private:
 	friend FArchive& operator<<(FArchive& Ar, class FVulkanShader::FSpirvContainer& SpirvContainer);
 	static FSpirvCode PatchSpirvInputAttachments(FSpirvCode& SpirvCode);
 
+	static FSpirvCode GetSpirvCode(const FSpirvContainer& Container);
+
 protected:
 	void Setup(FVulkanShaderHeader&& InCodeHeader, FShaderResourceTable&& InSRT, FSpirvContainer&& InSpirvContainer, uint64 InShaderKey);
 
@@ -271,9 +281,20 @@ private:
 		, FVulkanShader(InDevice, InFrequency)
 	{
 	}
+
+	FSpirvContainer AnyHitSpirvContainer;
+	FSpirvContainer IntersectionSpirvContainer;
+
 	friend class FVulkanShaderFactory;
 
 public:
+	static const uint32 MainModuleIdentifier = 0;
+	static const uint32 ClosestHitModuleIdentifier = MainModuleIdentifier;
+	static const uint32 AnyHitModuleIdentifier = 1;
+	static const uint32 IntersectionModuleIdentifier = 2;
+
+	TRefCountPtr<FVulkanShaderModule> GetOrCreateHandle(uint32 ModuleIdentifier);
+
 	// IRefCountedObject interface.
 	virtual uint32 AddRef() const override final
 	{
diff --git a/Engine/Source/Runtime/VulkanRHI/Private/VulkanShaders.cpp b/Engine/Source/Runtime/VulkanRHI/Private/VulkanShaders.cpp
--- a/Engine/Source/Runtime/VulkanRHI/Private/VulkanShaders.cpp
+++ b/Engine/Source/Runtime/VulkanRHI/Private/VulkanShaders.cpp
@@ -133,6 +133,21 @@ FVulkanRayTracingShader* FVulkanShaderFactory::CreateRayTracingShader(TArrayView
 		FVulkanShader::FSpirvContainer SpirvContainer;
 		Ar << SpirvContainer;
 
+		const bool bIsHitGroup = (ShaderFrequency == SF_RayHitGroup);
+		FVulkanShader::FSpirvContainer AnyHitSpirvContainer;
+		FVulkanShader::FSpirvContainer IntersectionSpirvContainer;
+		if (bIsHitGroup)
+		{
+			if (CodeHeader.RayGroupAnyHit == FVulkanShaderHeader::ERayHitGroupEntrypoint::SeparateBlob)
+			{
+				Ar << AnyHitSpirvContainer;
+			}
+			if (CodeHeader.RayGroupIntersection == FVulkanShaderHeader::ERayHitGroupEntrypoint::SeparateBlob)
+			{
+				Ar << IntersectionSpirvContainer;
+			}
+		}
+
 		{
 			FRWScopeLock ScopedLock(RWLock[ShaderFrequency], SLT_Write);
 			FVulkanShader* const* FoundShaderPtr = ShaderMap[ShaderFrequency].Find(ShaderKey);
@@ -144,6 +159,11 @@ FVulkanRayTracingShader* FVulkanShaderFactory::CreateRayTracingShader(TArrayView
 			{
 				RetShader = new FVulkanRayTracingShader(Device, ShaderFrequency);
 				RetShader->Setup(MoveTemp(CodeHeader), MoveTemp(SerializedSRT), MoveTemp(SpirvContainer), ShaderKey);
+				if (bIsHitGroup)
+				{
+					RetShader->AnyHitSpirvContainer = MoveTemp(AnyHitSpirvContainer);
+					RetShader->IntersectionSpirvContainer = MoveTemp(IntersectionSpirvContainer);
+				}
 				RetShader->RayTracingPayloadType = RetShader->CodeHeader.RayTracingPayloadType;
 				RetShader->RayTracingPayloadSize = RetShader->CodeHeader.RayTracingPayloadSize;
 
@@ -233,21 +253,21 @@ FVulkanShaderModule::~FVulkanShaderModule()
 	Device->GetDeferredDeletionQueue().EnqueueResource(VulkanRHI::FDeferredDeletionQueue2::EType::ShaderModule, ActualShaderModule);
 }
 
-FVulkanShader::FSpirvCode FVulkanShader::GetSpirvCode()
+FVulkanShader::FSpirvCode FVulkanShader::GetSpirvCode(const FSpirvContainer& Container)
 {
-	if (SpirvContainer.IsCompressed())
+	if (Container.IsCompressed())
 	{
 		TArray<uint32> UncompressedSpirv;
 		const size_t ElementSize = UncompressedSpirv.GetTypeSize();
-		UncompressedSpirv.Reserve(SpirvContainer.GetSizeBytes() / ElementSize);
-		UncompressedSpirv.SetNumUninitialized(SpirvContainer.GetSizeBytes() / ElementSize);
-		FCompression::UncompressMemory(NAME_Zlib, UncompressedSpirv.GetData(), SpirvContainer.GetSizeBytes(), SpirvContainer.SpirvCode.GetData(), SpirvContainer.SpirvCode.Num());
+		UncompressedSpirv.Reserve(Container.GetSizeBytes() / ElementSize);
+		UncompressedSpirv.SetNumUninitialized(Container.GetSizeBytes() / ElementSize);
+		FCompression::UncompressMemory(NAME_Zlib, UncompressedSpirv.GetData(), Container.GetSizeBytes(), Container.SpirvCode.GetData(), Container.SpirvCode.Num());
 
 		return FSpirvCode(MoveTemp(UncompressedSpirv));
 	}
 	else
 	{
-		return FSpirvCode(TArrayView<uint32>((uint32*)SpirvContainer.SpirvCode.GetData(), SpirvContainer.SpirvCode.Num() / sizeof(uint32)));
+		return FSpirvCode(TArrayView<uint32>((uint32*)Container.SpirvCode.GetData(), Container.SpirvCode.Num() / sizeof(uint32)));
 	}
 }
 
@@ -267,7 +287,7 @@ void FVulkanShader::Setup(FVulkanShaderHeader&& InCodeHeader, FShaderResourceTab
 
 	checkf(SpirvContainer.GetSizeBytes() != 0, TEXT("Empty SPIR-V! %s"), *CodeHeader.DebugName);
 
-	check(CodeHeader.UniformBufferSpirvInfos.Num() == CodeHeader.UniformBuffers.Num());
+	check(IsRayTracingShaderFrequency(Frequency) || (CodeHeader.UniformBufferSpirvInfos.Num() == CodeHeader.UniformBuffers.Num()));
 
 	check(CodeHeader.GlobalSpirvInfos.Num() == CodeHeader.Globals.Num());
 
@@ -433,10 +453,18 @@ TRefCountPtr<FVulkanShaderModule> FVulkanShader::GetOrCreateHandle()
 {
 	check(Device->SupportsBindless());
 	FScopeLock Lock(&VulkanShaderModulesMapCS);
+
+	const uint32 MainModuleIndex = 0;
+	TRefCountPtr<FVulkanShaderModule>* Found = ShaderModules.Find(MainModuleIndex);
+	if (Found)
+	{
+		return *Found;
+	}
+
 	FSpirvCode Spirv = GetSpirvCode();
 
 	TRefCountPtr<FVulkanShaderModule> Module = CreateShaderModule(Device, Spirv);
-	ShaderModules.Add(0, Module);
+	ShaderModules.Add(MainModuleIndex, Module);
 	if (!CodeHeader.DebugName.IsEmpty())
 	{
 		VULKAN_SET_DEBUG_NAME((*Device), VK_OBJECT_TYPE_SHADER_MODULE, Module->GetVkShaderModule(), TEXT("%s : (FVulkanShader*)0x%p"), *CodeHeader.DebugName, this);
@@ -444,6 +472,77 @@ TRefCountPtr<FVulkanShaderModule> FVulkanShader::GetOrCreateHandle()
 	return Module;
 }
 
+#if VULKAN_RHI_RAYTRACING
+TRefCountPtr<FVulkanShaderModule> FVulkanRayTracingShader::GetOrCreateHandle(uint32 ModuleIdentifier)
+{
+	check(Device->SupportsBindless());
+
+	const bool bIsAnyHitModuleIdentifier = (ModuleIdentifier == AnyHitModuleIdentifier);
+	const bool bIsIntersectionModuleIdentifier = (ModuleIdentifier == IntersectionModuleIdentifier);
+
+	// If we're using a single blob with multiple entry points, forward everything to the main module
+	if ((bIsAnyHitModuleIdentifier && (GetCodeHeader().RayGroupAnyHit == FVulkanShaderHeader::ERayHitGroupEntrypoint::CommonBlob)) ||
+		(bIsIntersectionModuleIdentifier && (GetCodeHeader().RayGroupIntersection == FVulkanShaderHeader::ERayHitGroupEntrypoint::CommonBlob)))
+	{
+		return GetOrCreateHandle(MainModuleIdentifier);
+	}
+
+	FScopeLock Lock(&VulkanShaderModulesMapCS);
+
+	TRefCountPtr<FVulkanShaderModule>* Found = ShaderModules.Find(ModuleIdentifier);
+	if (Found)
+	{
+		return *Found;
+	}
+
+	auto CreateHitGroupHandle = [&](const FSpirvContainer& Container)
+	{
+		FSpirvCode Spirv = GetSpirvCode(Container);
+		TRefCountPtr<FVulkanShaderModule> Module = CreateShaderModule(Device, Spirv);
+		ShaderModules.Add(ModuleIdentifier, Module);
+		return Module;
+	};
+
+	TRefCountPtr<FVulkanShaderModule> Module;
+	if (bIsAnyHitModuleIdentifier)
+	{
+		check(GetFrequency() == SF_RayHitGroup);
+		if (GetCodeHeader().RayGroupAnyHit == FVulkanShaderHeader::ERayHitGroupEntrypoint::SeparateBlob)
+		{
+			Module = CreateHitGroupHandle(AnyHitSpirvContainer);
+		}
+		else
+		{
+			return TRefCountPtr<FVulkanShaderModule>();
+		}
+	}
+	else if (bIsIntersectionModuleIdentifier)
+	{
+		check(GetFrequency() == SF_RayHitGroup);
+		if (GetCodeHeader().RayGroupIntersection == FVulkanShaderHeader::ERayHitGroupEntrypoint::SeparateBlob)
+		{
+			Module = CreateHitGroupHandle(IntersectionSpirvContainer);
+		}
+		else
+		{
+			return TRefCountPtr<FVulkanShaderModule>();
+		}
+	}
+	else
+	{
+		Module = CreateHitGroupHandle(SpirvContainer);
+	}
+
+	if (!CodeHeader.DebugName.IsEmpty())
+	{
+		VULKAN_SET_DEBUG_NAME((*Device), VK_OBJECT_TYPE_SHADER_MODULE, Module->GetVkShaderModule(), TEXT("%s : (FVulkanShader*)0x%p"), *CodeHeader.DebugName, this);
+	}
+
+	return Module;
+}
+#endif // VULKAN_RHI_RAYTRACING
+
+
 TRefCountPtr<FVulkanShaderModule> FVulkanShader::CreateHandle(const FVulkanLayout* Layout, uint32 LayoutHash)
 {
 	FScopeLock Lock(&VulkanShaderModulesMapCS);
