VulkanRHI:  Actually enable scalar_block_layout in raytracing and sm6 shaders to fix multiple alignment issues.  Rename ReflectionSlot param to align with ShaderCommon code.

#rnx

diff --git a/Engine/Shaders/Public/Platform/Vulkan/VulkanCommon.ush b/Engine/Shaders/Public/Platform/Vulkan/VulkanCommon.ush
--- a/Engine/Shaders/Public/Platform/Vulkan/VulkanCommon.ush
+++ b/Engine/Shaders/Public/Platform/Vulkan/VulkanCommon.ush
@@ -5,7 +5,7 @@
 =============================================================================*/
 
 // Update this GUID to improve shader recompilation for Vulkan only shaders
-// GUID = EBB16B3A-3FC4-4518-BDF0-63FCF6281A9A
+// GUID = 8931F79A-3E7B-4BB4-BB26-341699B0D77E
 
 #pragma once
 
diff --git a/Engine/Source/Developer/ShaderCompilerCommon/Private/ShaderConductorContext.cpp b/Engine/Source/Developer/ShaderCompilerCommon/Private/ShaderConductorContext.cpp
--- a/Engine/Source/Developer/ShaderCompilerCommon/Private/ShaderConductorContext.cpp
+++ b/Engine/Source/Developer/ShaderCompilerCommon/Private/ShaderConductorContext.cpp
@@ -375,9 +375,9 @@ namespace CrossCompiler
 		}
 
 		// Add additional DXC arguments that are not exposed by ShaderConductor API directly
-		if (InOptions.bDisableScalarBlockLayout)
+		if (!InOptions.bDisableScalarBlockLayout)
 		{
-			DxcArgRefs.Add("-fspv-no-scalar-block-layout");
+			DxcArgRefs.Add("-fvk-use-scalar-layout");
 		}
 		if (InOptions.bPreserveStorageInput)
 		{
diff --git a/Engine/Source/Developer/VulkanShaderFormat/Private/VulkanShaderCompiler.cpp b/Engine/Source/Developer/VulkanShaderFormat/Private/VulkanShaderCompiler.cpp
--- a/Engine/Source/Developer/VulkanShaderFormat/Private/VulkanShaderCompiler.cpp
+++ b/Engine/Source/Developer/VulkanShaderFormat/Private/VulkanShaderCompiler.cpp
@@ -1378,7 +1378,7 @@ static bool BuildShaderOutputFromSpirv(
 	RegisterBindings(Bindings.TextureUAVs, "u", EVulkanBindingType::StorageImage);
 
 	RegisterBindings(Bindings.TBufferSRVs, "s", EVulkanBindingType::UniformTexelBuffer);
-	RegisterBindings(Bindings.SBufferSRVs, "s", EVulkanBindingType::UniformTexelBuffer);
+	checkf(Bindings.SBufferSRVs.IsEmpty(), TEXT("GatherSpirvReflectionBindings should have dumped all SBufferSRVs into SBufferUAVs."));
 	RegisterBindings(Bindings.TextureSRVs, "s", EVulkanBindingType::Image);
 
 	RegisterBindings(Bindings.Samplers, "z", EVulkanBindingType::Sampler);
@@ -1420,7 +1420,7 @@ static bool BuildShaderOutputFromSpirv(
 				const SpvReflectResult SpvResult = Reflection.ChangeDescriptorBindingNumbers(Binding, BindingIndex, DescSetNumber);
 				check(SpvResult == SPV_REFLECT_RESULT_SUCCESS);
 
-				const int32 VulkanBindingIndex = SerializedOutput.Spirv.ReflectionInfo.Add(FVulkanSpirv::FEntry(ResourceName, BindingIndex));
+				const int32 ReflectionSlot = SerializedOutput.Spirv.ReflectionInfo.Add(FVulkanSpirv::FEntry(ResourceName, BindingIndex));
 				const FShaderParameterParser::FParsedShaderParameter* ParsedParam = InternalState.ParameterParser.FindParameterInfosUnsafe(ResourceName);
 
 				auto AddShaderValidationType = [] (uint32_t VulkanBindingIndex, const FShaderParameterParser::FParsedShaderParameter* ParsedParam, FShaderCompilerOutput& Output) {
@@ -1442,7 +1442,7 @@ static bool BuildShaderOutputFromSpirv(
 				case EVulkanBindingType::StorageTexelBuffer:
 				case EVulkanBindingType::StorageBuffer:
 				case EVulkanBindingType::StorageImage:
-					HandleReflectedShaderUAV(ResourceName, BindingOffset, VulkanBindingIndex, 1, Output);
+					HandleReflectedShaderUAV(ResourceName, BindingOffset, ReflectionSlot, 1, Output);
 					EntryTypes.Add(ResourceName, FVulkanShaderHeader::Global);
 
 					AddShaderValidationType(BindingOffset, ParsedParam, Output);
@@ -1458,20 +1458,20 @@ static bool BuildShaderOutputFromSpirv(
 					[[fallthrough]];
 
 				case EVulkanBindingType::UniformTexelBuffer:
-					HandleReflectedShaderResource(ResourceName, BindingOffset, VulkanBindingIndex, 1, Output);
+					HandleReflectedShaderResource(ResourceName, BindingOffset, ReflectionSlot, 1, Output);
 					EntryTypes.Add(ResourceName, FVulkanShaderHeader::Global);
 
 					AddShaderValidationType(BindingOffset, ParsedParam, Output);
 					break;
 
 				case EVulkanBindingType::Sampler:
-					HandleReflectedShaderSampler(ResourceName, VulkanBindingIndex, Output);
-					//HandleReflectedShaderSampler(ResourceName, BindingOffset, VulkanBindingIndex, 1, Output);
+					HandleReflectedShaderSampler(ResourceName, ReflectionSlot, Output);
+					//HandleReflectedShaderSampler(ResourceName, BindingOffset, ReflectionSlot, 1, Output);
 					EntryTypes.Add(ResourceName, FVulkanShaderHeader::Global);
 					break;
 
 				case EVulkanBindingType::AccelerationStructure:
-					HandleReflectedShaderResource(ResourceName, BindingOffset, VulkanBindingIndex, 1, Output);
+					HandleReflectedShaderResource(ResourceName, BindingOffset, ReflectionSlot, 1, Output);
 					EntryTypes.Add(ResourceName, FVulkanShaderHeader::Global);
 
 					AddShaderValidationType(BindingOffset, ParsedParam, Output);
@@ -1484,13 +1484,13 @@ static bool BuildShaderOutputFromSpirv(
 				case EVulkanBindingType::PackedUniformBuffer:
 					{
 						// Use the given global ResourceName instead of patching it to _Globals_h
-						check(!UsedUniformBufferSlots[VulkanBindingIndex]);
-						UsedUniformBufferSlots[VulkanBindingIndex] = true;
+						check(!UsedUniformBufferSlots[ReflectionSlot]);
+						UsedUniformBufferSlots[ReflectionSlot] = true;
 
 						if (InternalState.UseRootParametersStructure())
 						{
-							check(VulkanBindingIndex == FShaderParametersMetadata::kRootCBufferBindingIndex);
-							HandleReflectedUniformBuffer(ResourceName, VulkanBindingIndex, Output);
+							check(ReflectionSlot == FShaderParametersMetadata::kRootCBufferBindingIndex);
+							HandleReflectedUniformBuffer(ResourceName, ReflectionSlot, Output);
 							EntryTypes.Add(ResourceName, FVulkanShaderHeader::UniformBuffer);
 						}
 
@@ -1524,9 +1524,9 @@ static bool BuildShaderOutputFromSpirv(
 
 				case EVulkanBindingType::UniformBuffer:
 					{
-						check(!UsedUniformBufferSlots[VulkanBindingIndex]);
-						UsedUniformBufferSlots[VulkanBindingIndex] = true;
-						HandleReflectedUniformBuffer(ResourceName, VulkanBindingIndex, Output);
+						check(!UsedUniformBufferSlots[ReflectionSlot]);
+						UsedUniformBufferSlots[ReflectionSlot] = true;
+						HandleReflectedUniformBuffer(ResourceName, ReflectionSlot, Output);
 						EntryTypes.Add(ResourceName, FVulkanShaderHeader::UniformBuffer);
 
 						AddShaderValidationUBSize(BindingOffset, Binding->block.padded_size, Output);
@@ -1556,7 +1556,7 @@ static bool BuildShaderOutputFromSpirv(
 
 		int32 SRVBindings = 0;
 		SRVBindings = AddReflectionInfos(Bindings.TBufferSRVs, EVulkanBindingType::UniformTexelBuffer, SRVBindings);
-		SRVBindings = AddReflectionInfos(Bindings.SBufferSRVs, EVulkanBindingType::UniformTexelBuffer, SRVBindings);
+		checkf(Bindings.SBufferSRVs.IsEmpty(), TEXT("GatherSpirvReflectionBindings should have dumped all SBufferSRVs into SBufferUAVs."));
 		SRVBindings = AddReflectionInfos(Bindings.TextureSRVs, EVulkanBindingType::Image, SRVBindings);
 
 		Output.NumTextureSamplers = AddReflectionInfos(Bindings.Samplers, EVulkanBindingType::Sampler, 0);
@@ -2033,8 +2033,7 @@ static bool CompileWithShaderConductor(
 	Options.TargetEnvironment = InternalState.MinimumTargetEnvironment;
 
 	// VK_EXT_scalar_block_layout is required by raytracing and by Nanite (so expect it to be present in SM6/Vulkan_1_3)
-	Options.bDisableScalarBlockLayout = !(InternalState.IsRayTracingShader() ||
-		InternalState.MinimumTargetEnvironment >= CrossCompiler::FShaderConductorOptions::ETargetEnvironment::Vulkan_1_3);
+	Options.bDisableScalarBlockLayout = !(InternalState.IsRayTracingShader() || InternalState.IsSM6());
 
 	if (Input.Environment.CompilerFlags.Contains(CFLAG_AllowRealTypes))
 	{
