VulkanRHI: VK RT WIP: Use tighly packed indices for the bindless UBs in the shader parameters, and at runtime use ConstantDataOriginalBindingIndex to write the handle in the correct slot.

#rnx

diff --git a/Engine/Source/Developer/VulkanShaderFormat/Private/VulkanShaderCompiler.cpp b/Engine/Source/Developer/VulkanShaderFormat/Private/VulkanShaderCompiler.cpp
--- a/Engine/Source/Developer/VulkanShaderFormat/Private/VulkanShaderCompiler.cpp
+++ b/Engine/Source/Developer/VulkanShaderFormat/Private/VulkanShaderCompiler.cpp
@@ -2382,6 +2382,49 @@ static TArray<FString> ConvertUBToBindless(FString& PreprocessedShaderSource)
 }
 
 
+static void UpdateBindlessUBs(const FVulkanShaderCompilerInternalState& InternalState, VulkanShaderCompilerSerializedOutput& SerializedOutput, FShaderCompilerOutput& Output)
+{
+	auto GetLayoutHash = [&InternalState](const FString& UBName)
+	{
+		uint32 LayoutHash = 0;
+		const FUniformBufferEntry* UniformBufferEntry = InternalState.Input.Environment.UniformBufferMap.Find(UBName);
+		if (UniformBufferEntry)
+		{
+			LayoutHash = UniformBufferEntry->LayoutHash;
+		}
+		else if ((UBName == FShaderParametersMetadata::kRootUniformBufferBindingName) && InternalState.Input.RootParametersStructure)
+		{
+			LayoutHash = InternalState.Input.RootParametersStructure->GetLayoutHash();
+		}
+		else
+		{
+			LayoutHash = 0;
+		}
+		return LayoutHash;
+	};
+
+	SerializedOutput.Header.UniformBuffers.Empty();
+	for (int32 CBIndex = 0; CBIndex < InternalState.AllBindlessUBs.Num(); CBIndex++)
+	{
+		const FString& CBName = InternalState.AllBindlessUBs[CBIndex];
+
+		// It's possible SPIRV compilation has optimized out a buffer from every shader in the group
+		if (SerializedOutput.UsedBindlessUB.Contains(CBName))
+		{
+			FVulkanShaderHeader::FUniformBufferInfo& UBInfo = SerializedOutput.Header.UniformBuffers.AddZeroed_GetRef();
+			UBInfo.LayoutHash = GetLayoutHash(CBName);
+			UBInfo.ConstantDataOriginalBindingIndex = CBIndex;
+#if VULKAN_ENABLE_BINDING_DEBUG_NAMES
+			UBInfo.DebugName = CBName;
+#endif
+
+			const int32 UBIndex = SerializedOutput.Header.UniformBuffers.Num() - 1;
+			Output.ParameterMap.AddParameterAllocation(*CBName, UBIndex, (uint16)FVulkanShaderHeader::UniformBuffer, 1, EShaderParameterType::UniformBuffer);
+		}
+	}
+}
+
+
 static bool CompileShaderGroup(
 	FVulkanShaderCompilerInternalState& InternalState,
 	const FString& OriginalPreprocessedShaderSource,
@@ -2446,48 +2489,10 @@ static bool CompileShaderGroup(
 		bSuccess = CompilePartialExport(FVulkanShaderCompilerInternalState::EHitGroupShaderType::Intersection, TEXT("intersection"), IntersectionSerializedOutput);
 	}
 
-	// Add the bindless data in the output at the very end
-	{
-		auto GetLayoutHash = [&InternalState](const FString& UBName)
-		{
-			uint32 LayoutHash = 0;
-			const FUniformBufferEntry* UniformBufferEntry = InternalState.Input.Environment.UniformBufferMap.Find(UBName);
-			if (UniformBufferEntry)
-			{
-				LayoutHash = UniformBufferEntry->LayoutHash;
-			}
-			else if ((UBName == FShaderParametersMetadata::kRootUniformBufferBindingName) && InternalState.Input.RootParametersStructure)
-			{
-				LayoutHash = InternalState.Input.RootParametersStructure->GetLayoutHash();
-			}
-			else
-			{
-				LayoutHash = 0;
-			}
-			return LayoutHash;
-		};
-
-		ClosestHitSerializedOutput.Header.UniformBuffers.Empty();
-		for (int32 CBIndex = 0; CBIndex < InternalState.AllBindlessUBs.Num(); CBIndex++)
-		{
-			const FString& CBName = InternalState.AllBindlessUBs[CBIndex];
-
-			// It's possible SPIRV compilation has optimized out a buffer from every shader in the group
-			if (ClosestHitSerializedOutput.UsedBindlessUB.Contains(CBName) ||
-				AnyHitSerializedOutput.UsedBindlessUB.Contains(CBName) ||
-				IntersectionSerializedOutput.UsedBindlessUB.Contains(CBName))
-			{
-				FVulkanShaderHeader::FUniformBufferInfo& UBInfo = ClosestHitSerializedOutput.Header.UniformBuffers.AddZeroed_GetRef();
-				UBInfo.LayoutHash = GetLayoutHash(CBName);
-				UBInfo.ConstantDataOriginalBindingIndex = CBIndex;
-#if VULKAN_ENABLE_BINDING_DEBUG_NAMES
-				UBInfo.DebugName = CBName;
-#endif
-
-				MergedOutput.ParameterMap.AddParameterAllocation(*CBName, CBIndex, (uint16)FVulkanShaderHeader::UniformBuffer, 1, EShaderParameterType::UniformBuffer);
-			}
-		}
-	}
+	// Collapse the bindless UB usage into one set and then update the headers
+	ClosestHitSerializedOutput.UsedBindlessUB.Append(AnyHitSerializedOutput.UsedBindlessUB);
+	ClosestHitSerializedOutput.UsedBindlessUB.Append(IntersectionSerializedOutput.UsedBindlessUB);
+	UpdateBindlessUBs(InternalState, ClosestHitSerializedOutput, MergedOutput);
 
 	{
 		// :todo-jn: Having multiple entrypoints in a single SPIRV blob crashes on FLumenHardwareRayTracingMaterialHitGroup for some reason
@@ -2564,6 +2569,11 @@ void CompileVulkanShader(const FShaderCompilerInput& Input, const FShaderPreproc
 		VulkanShaderCompilerSerializedOutput SerializedOutput;
 		bSuccess = CompileWithShaderConductor(InternalState, PreprocessedShaderSource, SerializedOutput, Output);
 
+		if (InternalState.bUseBindlessUniformBuffer)
+		{
+			UpdateBindlessUBs(InternalState, SerializedOutput, Output);
+		}
+
 		// Write out the header and shader source code (except for the extra shaders in hit groups)
 		check(SerializedOutput.Spirv.Data.Num() != 0);
 		FMemoryWriter Ar(Output.ShaderCode.GetWriteAccess(), true);
diff --git a/Engine/Source/Runtime/VulkanRHI/Private/VulkanRayTracing.cpp b/Engine/Source/Runtime/VulkanRHI/Private/VulkanRayTracing.cpp
--- a/Engine/Source/Runtime/VulkanRHI/Private/VulkanRayTracing.cpp
+++ b/Engine/Source/Runtime/VulkanRHI/Private/VulkanRayTracing.cpp
@@ -1887,25 +1887,20 @@ void FVulkanCommandListContext::RHIRayTraceDispatchIndirect(
 static void SetSystemParametersUB(FVulkanHitGroupSystemParameters& OutSystemParameters, uint32 InNumUniformBuffers, FRHIUniformBuffer* const* InUniformBuffers, const FVulkanRayTracingShader* InShader)
 {
 	// Plug the shaders in the right slots using LayoutHash comparisons
-	// :todo-jn: Find a more direct approach
+	check(InShader->GetCodeHeader().UniformBuffers.Num() == InNumUniformBuffers);
 	for (uint32 UBIndex = 0; UBIndex < InNumUniformBuffers; ++UBIndex)
 	{
 		FVulkanUniformBuffer* UniformBuffer = ResourceCast(InUniformBuffers[UBIndex]);
-		if (UniformBuffer)
-		{
-			const uint32 LayoutHash = UniformBuffer->GetLayout().GetHash();
+		check(UniformBuffer);
 
-			for (int32 TargetIndex = 0; TargetIndex < InShader->GetCodeHeader().UniformBuffers.Num(); ++TargetIndex)
-			{
-				if (InShader->GetCodeHeader().UniformBuffers[TargetIndex].LayoutHash == LayoutHash)
-				{
-					const FRHIDescriptorHandle BindlessHandle = UniformBuffer->GetBindlessHandle();
-					check(BindlessHandle.IsValid());
-					OutSystemParameters.BindlessUniformBuffers[TargetIndex] = BindlessHandle.GetIndex();
-					break;
-				}
-			}
-		}
+		const FVulkanShaderHeader::FUniformBufferInfo& UniformBufferInfo = InShader->GetCodeHeader().UniformBuffers[UBIndex];
+		check((UniformBufferInfo.LayoutHash == 0) || (UniformBufferInfo.LayoutHash == UniformBuffer->GetLayout().GetHash()));
+		check(UniformBufferInfo.ConstantDataOriginalBindingIndex != UINT16_MAX);
+
+
+		const FRHIDescriptorHandle BindlessHandle = UniformBuffer->GetBindlessHandle();
+		check(BindlessHandle.IsValid());
+		OutSystemParameters.BindlessUniformBuffers[UniformBufferInfo.ConstantDataOriginalBindingIndex] = BindlessHandle.GetIndex();
 	}
 }
 
