VulkanRHI: Fix Vulkan bindless and ray tracing shader compilation to work with new UB generation code.

#rnx

diff --git a/Engine/Shaders/Public/Platform/Vulkan/VulkanCommon.ush b/Engine/Shaders/Public/Platform/Vulkan/VulkanCommon.ush
--- a/Engine/Shaders/Public/Platform/Vulkan/VulkanCommon.ush
+++ b/Engine/Shaders/Public/Platform/Vulkan/VulkanCommon.ush
@@ -52,9 +52,9 @@ void ImageInterlockedMaxUInt64(RWTexture2D<UlongType> Dest, uint2 Coord, UlongTy
 
 	#define UB_RESOURCE_MEMBER_SAMPLER(Type, UBName, Name) \
 		typedef Type SafeType##UBName##_##Name; \
-		SafeType##UBName##_##Name VULKAN_RESOURCE_HEAP(SafeType##UBName##_##Name)[]; \
+		SafeType##UBName##_##Name VULKAN_SAMPLER_HEAP(SafeType##UBName##_##Name)[]; \
 		Type GetBindlessSampler##UBName##_##Name() { \
-			return VULKAN_RESOURCE_HEAP(SafeType##UBName##_##Name)[UB_CB_MEMBER_ACCESS(UBName, BindlessSampler_##Name)]; \
+			return VULKAN_SAMPLER_HEAP(SafeType##UBName##_##Name)[UB_CB_MEMBER_ACCESS(UBName, BindlessSampler_##Name)]; \
 		} \
 		static const Type UBName##_##Name = GetBindlessSampler##UBName##_##Name();
 #endif
@@ -88,9 +88,7 @@ void ImageInterlockedMaxUInt64(RWTexture2D<UlongType> Dest, uint2 Coord, UlongTy
 
 	static const FHitGroupSystemRootConstants HitGroupSystemRootConstants = VulkanHitGroupSystemParameters.RootConstants;
 
-	static const uint BindlessResource_HitGroupSystemIndexBuffer = VulkanHitGroupSystemParameters.BindlessHitGroupSystemIndexBuffer;
-	UB_RESOURCE_MEMBER_RESOURCE(ByteAddressBuffer, HitGroupSystemIndexBuffer);
-	
-	static const uint BindlessResource_HitGroupSystemVertexBuffer = VulkanHitGroupSystemParameters.BindlessHitGroupSystemVertexBuffer;
-	UB_RESOURCE_MEMBER_RESOURCE(ByteAddressBuffer, HitGroupSystemVertexBuffer);
+	ByteAddressBuffer VULKAN_RESOURCE_HEAP(HitGroupSystemBuffers)[];
+	static const ByteAddressBuffer HitGroupSystemIndexBuffer = VULKAN_RESOURCE_HEAP(HitGroupSystemBuffers)[VulkanHitGroupSystemParameters.BindlessHitGroupSystemIndexBuffer];
+	static const ByteAddressBuffer HitGroupSystemVertexBuffer = VULKAN_RESOURCE_HEAP(HitGroupSystemBuffers)[VulkanHitGroupSystemParameters.BindlessHitGroupSystemVertexBuffer];
 #endif
diff --git a/Engine/Source/Developer/VulkanShaderFormat/Private/VulkanShaderCompiler.cpp b/Engine/Source/Developer/VulkanShaderFormat/Private/VulkanShaderCompiler.cpp
--- a/Engine/Source/Developer/VulkanShaderFormat/Private/VulkanShaderCompiler.cpp
+++ b/Engine/Source/Developer/VulkanShaderFormat/Private/VulkanShaderCompiler.cpp
@@ -2287,6 +2287,7 @@ static TArray<FString> ConvertUBToBindless(FString& PreprocessedShaderSource)
 		const FString PrefixedCBName = FString::Printf(TEXT("%s%d_%s"), *kBindlessCBPrefix, CBIndex, *CBName);
 		const FString BindlessCBType = PrefixedCBName + TEXT("_Type");
 		const FString BindlessCBHeapName = PrefixedCBName + kBindlessHeapSuffix;
+		const FString PaddingName = FString::Printf(TEXT("%s_Padding"), *CBName);
 
 		FString CBDecl;
 		CBDecl.Reserve(Members.Len() * 3);  // start somewhere approx less bad
@@ -2319,28 +2320,40 @@ static TArray<FString> ConvertUBToBindless(FString& PreprocessedShaderSource)
 			MemberSearchPtr = ParseHLSLSymbolName(MemberTypeEndPtr, MemberName);
 			check(MemberName.Len() > 0);
 
-			// Skip over trailing tokens and pick up arrays
-			FString ArrayDecl;
-			while (*MemberSearchPtr && *MemberSearchPtr != ';')
+			if (MemberName.StartsWith(PaddingName))
 			{
-				if (*MemberSearchPtr == '[')
+				while (*MemberSearchPtr && *MemberSearchPtr != ';')
 				{
-					ArrayDecl.AppendChar(*MemberSearchPtr);
-
 					MemberSearchPtr++;
-					while (*MemberSearchPtr && *MemberSearchPtr != ']')
+				}
+			}
+			else
+			{
+				// Skip over trailing tokens and pick up arrays
+				FString ArrayDecl;
+				while (*MemberSearchPtr && *MemberSearchPtr != ';')
+				{
+					if (*MemberSearchPtr == '[')
 					{
 						ArrayDecl.AppendChar(*MemberSearchPtr);
+
 						MemberSearchPtr++;
+						while (*MemberSearchPtr && *MemberSearchPtr != ']')
+						{
+							ArrayDecl.AppendChar(*MemberSearchPtr);
+							MemberSearchPtr++;
+						}
+
+						ArrayDecl.AppendChar(*MemberSearchPtr);
 					}
 
-					ArrayDecl.AppendChar(*MemberSearchPtr);
+					MemberSearchPtr++;
 				}
 
-				MemberSearchPtr++;
+				CBDecl += FString::Printf(TEXT("static const %s %s%s = %s.%s;\n"), *MemberTypeName, *MemberName, *ArrayDecl, *PrefixedCBName, *MemberName);
 			}
 
-			CBDecl += FString::Printf(TEXT("static const %s %s%s = %s.%s;\n"), *MemberTypeName, *MemberName, *ArrayDecl, *PrefixedCBName, *MemberName);
+			MemberSearchPtr++;
 
 		} while (MemberSearchPtr < LastMemberSemicolon);
 
