Moving more platforms to use ShaderParameterParser in the compile step.

diff --git a/Engine/Shaders/Public/ShaderVersion.ush b/Engine/Shaders/Public/ShaderVersion.ush
--- a/Engine/Shaders/Public/ShaderVersion.ush
+++ b/Engine/Shaders/Public/ShaderVersion.ush
@@ -2,4 +2,4 @@
 // Each time the console command is executed it generates a new GUID. As this file is included
 // in Platform.ush (which should be included in any shader) it allows to invalidate the shader DDC.
 // 
-#pragma message("UESHADERMETADATA_VERSION BC047A48-A9B5-4CF5-BD87-B5E9107C5711")
+#pragma message("UESHADERMETADATA_VERSION 46D2055A-DC11-404B-858A-B53932CBCD86")
diff --git a/Engine/Source/Developer/Apple/MetalShaderFormat/Private/MetalShaderCompiler.cpp b/Engine/Source/Developer/Apple/MetalShaderFormat/Private/MetalShaderCompiler.cpp
--- a/Engine/Source/Developer/Apple/MetalShaderFormat/Private/MetalShaderCompiler.cpp
+++ b/Engine/Source/Developer/Apple/MetalShaderFormat/Private/MetalShaderCompiler.cpp
@@ -872,12 +872,6 @@ bool PreprocessMetalShader(const FShaderCompilerInput& Input, const FShaderCompi
 		}
 	}
 
-	if (!PreprocessOutput.ParseAndModify(Input, Environment, nullptr))
-	{
-		// The FShaderParameterParser will add any relevant errors.
-		return false;
-	}
-
 	CleanupUniformBufferCode(Input.Environment, PreprocessOutput.EditSource());
 
 	// Process TEXT macro.
@@ -894,8 +888,23 @@ bool PreprocessMetalShader(const FShaderCompilerInput& Input, const FShaderCompi
 	return true;
 }
 
-void CompileMetalShader(const FShaderCompilerInput& Input, const FShaderPreprocessOutput& PreprocessOutput, FShaderCompilerOutput& Output)
+void CompileMetalShader(const FShaderCompilerInput& Input, const FString& InPreprocessedSource, FShaderCompilerOutput& Output)
 {
+	FString EntryPointName = Input.EntryPointName;
+	FString PreprocessedSource = InPreprocessedSource;
+
+	FShaderParameterParser ShaderParameterParser(Input.Environment.CompilerFlags, nullptr, {}, {});
+	if (!ShaderParameterParser.ParseAndModify(Input, Output.Errors, PreprocessedSource, EBindlessParameterMode::Default))
+	{
+		// The FShaderParameterParser will add any relevant errors.
+		return;
+	}
+
+	if (ShaderParameterParser.DidModifyShader())
+	{
+		Output.ModifiedShaderSource = PreprocessedSource;
+	}
+
 	EMetalGPUSemantics Semantics = EMetalGPUSemanticsMobile;
 	if (Input.ShaderFormat == NAME_SF_METAL_MRT
 		|| Input.ShaderFormat == NAME_SF_METAL_MRT_TVOS
@@ -1104,8 +1113,8 @@ void CompileMetalShader(const FShaderCompilerInput& Input, const FShaderPreproce
 		FSHA1::HashBuffer(&Guid, sizeof(FGuid), GUIDHash.Hash);
 	}
 
-	DoCompileMetalShader(Input, Output, PreprocessOutput.GetSource(), GUIDHash, VersionEnum, Semantics, MaxUnrollLoops, (EShaderFrequency)Input.Target.Frequency, bDumpDebugInfo, Standard, MinOSVersion);
-	PreprocessOutput.GetParameterParser().ValidateShaderParameterTypes(Input, bIsMobile, Output);
+	DoCompileMetalShader(Input, Output, PreprocessedSource, GUIDHash, VersionEnum, Semantics, MaxUnrollLoops, (EShaderFrequency)Input.Target.Frequency, bDumpDebugInfo, Standard, MinOSVersion);
+	ShaderParameterParser.ValidateShaderParameterTypes(Input, bIsMobile, Output);
 }
 
 void OutputMetalDebugData(const FShaderCompilerInput& Input, const FShaderPreprocessOutput& PreprocessOutput, const FShaderCompilerOutput& Output)
diff --git a/Engine/Source/Developer/Apple/MetalShaderFormat/Private/MetalShaderFormat.cpp b/Engine/Source/Developer/Apple/MetalShaderFormat/Private/MetalShaderFormat.cpp
--- a/Engine/Source/Developer/Apple/MetalShaderFormat/Private/MetalShaderFormat.cpp
+++ b/Engine/Source/Developer/Apple/MetalShaderFormat/Private/MetalShaderFormat.cpp
@@ -7,6 +7,7 @@
 #include "Interfaces/IShaderFormatModule.h"
 #include "ShaderCore.h"
 #include "ShaderCodeArchive.h"
+#include "ShaderPreprocessTypes.h"
 #include "hlslcc.h"
 #include "MetalShaderResources.h"
 #include "HAL/FileManager.h"
@@ -28,7 +29,7 @@ DEFINE_LOG_CATEGORY(LogMetalShaderCompiler)
 #define CHECK_METAL_COMPILER_TOOLCHAIN_SETUP 0
 
 extern bool PreprocessMetalShader(const FShaderCompilerInput& Input, const FShaderCompilerEnvironment& Environment, FShaderPreprocessOutput& PreprocessOutput);
-extern void CompileMetalShader(const FShaderCompilerInput& Input, const FShaderPreprocessOutput& PreprocessOutput, FShaderCompilerOutput& Output);
+extern void CompileMetalShader(const FShaderCompilerInput& Input, const FString& InPreprocessedSource, FShaderCompilerOutput& Output);
 extern void OutputMetalDebugData(const FShaderCompilerInput& Input, const FShaderPreprocessOutput& PreprocessOutput, const FShaderCompilerOutput& Output);
 
 extern bool StripShader_Metal(TArray<uint8>& Code, class FString const& DebugPath, bool const bNative);
@@ -164,7 +165,7 @@ public:
 		const FString& WorkingDirectory) const override final
 	{
 		CheckShaderFormat(Input.ShaderFormat);
-		CompileMetalShader(Input, PreprocessOutput, Output);
+		CompileMetalShader(Input, PreprocessOutput.GetSource(), Output);
 	}
 
 	virtual void OutputDebugData(
diff --git a/Engine/Source/Developer/ShaderFormatOpenGL/Private/OpenGLShaderCompiler.cpp b/Engine/Source/Developer/ShaderFormatOpenGL/Private/OpenGLShaderCompiler.cpp
--- a/Engine/Source/Developer/ShaderFormatOpenGL/Private/OpenGLShaderCompiler.cpp
+++ b/Engine/Source/Developer/ShaderFormatOpenGL/Private/OpenGLShaderCompiler.cpp
@@ -3291,12 +3291,6 @@ bool PreprocessOpenGLShader(const FShaderCompilerInput& Input, const FShaderComp
 		}
 	}
 
-	if (!PreprocessOutput.ParseAndModify(Input, Environment, nullptr))
-	{
-		// The FShaderParameterParser will add any relevant errors.
-		return false;
-	}
-
 	// This requires removing the HLSLCC_NoPreprocess flag later on!
 	CleanupUniformBufferCode(Input.Environment, PreprocessOutput.EditSource());
 
@@ -3321,9 +3315,23 @@ bool PreprocessOpenGLShader(const FShaderCompilerInput& Input, const FShaderComp
  * @param WorkingDirectory - (unused, part of IShaderFormat API)
  * @param Version - Target GLSL version for this compilation
  */
-void CompileOpenGLShader(const FShaderCompilerInput& Input, const FShaderPreprocessOutput& PreprocessOutput, FShaderCompilerOutput& Output, const FString& WorkingDirectory, GLSLVersion Version)
+void CompileOpenGLShader(const FShaderCompilerInput& Input, const FString& InPreprocessedSource, FShaderCompilerOutput& Output, const FString& WorkingDirectory, GLSLVersion Version)
 {
-	const FString& PreprocessedShader = PreprocessOutput.GetSource();
+	FString EntryPointName = Input.EntryPointName;
+	FString PreprocessedSource = InPreprocessedSource;
+
+	FShaderParameterParser ShaderParameterParser(Input.Environment.CompilerFlags, nullptr, {}, {});
+	if (!ShaderParameterParser.ParseAndModify(Input, Output.Errors, PreprocessedSource, EBindlessParameterMode::Default))
+	{
+		// The FShaderParameterParser will add any relevant errors.
+		return;
+	}
+
+	if (ShaderParameterParser.DidModifyShader())
+	{
+		Output.ModifiedShaderSource = PreprocessedSource;
+	}
+
 	const EHlslCompileTarget HlslCompilerTarget = GetCompileTarget(Version);
 
 	const bool bUseSC = ShouldUseDXC(Input.Environment.CompilerFlags);
@@ -3366,7 +3374,7 @@ void CompileOpenGLShader(const FShaderCompilerInput& Input, const FShaderPreproc
 #if DXC_SUPPORTED
 	if (bUseSC)
 	{
-		bCompilationSucceeded = CompileToGlslWithShaderConductor(Input, Output, Version, Frequency, PreprocessedShader, GlslShaderSource);
+		bCompilationSucceeded = CompileToGlslWithShaderConductor(Input, Output, Version, Frequency, PreprocessedSource, GlslShaderSource);
 	}
 	else
 #endif // DXC_SUPPORTED
@@ -3383,7 +3391,7 @@ void CompileOpenGLShader(const FShaderCompilerInput& Input, const FShaderPreproc
 			if (CrossCompilerContext.Init(TCHAR_TO_ANSI(*Input.VirtualSourceFilePath), LanguageSpec))
 			{
 				bCompilationSucceeded = CrossCompilerContext.Run(
-					TCHAR_TO_ANSI(*PreprocessedShader),
+					TCHAR_TO_ANSI(*PreprocessedSource),
 					TCHAR_TO_ANSI(*Input.EntryPointName),
 					BackEnd,
 					&GlslShaderSource,
@@ -3460,7 +3468,7 @@ void FOpenGLFrontend::CompileShader(const FShaderCompilerInput& Input, FShaderCo
 {
 	FShaderPreprocessOutput PreprocessOutput;
 	PreprocessOpenGLShader(Input, Input.Environment, PreprocessOutput, Version);
-	CompileOpenGLShader(Input, PreprocessOutput, Output, WorkingDirectory, Version);
+	CompileOpenGLShader(Input, PreprocessOutput.GetSource(), Output, WorkingDirectory, Version);
 }
 
 static void FillDeviceCapsOfflineCompilationInternal(struct FDeviceCapabilities& Capabilities, const GLSLVersion ShaderVersion)
diff --git a/Engine/Source/Developer/ShaderFormatOpenGL/Private/ShaderFormatOpenGL.cpp b/Engine/Source/Developer/ShaderFormatOpenGL/Private/ShaderFormatOpenGL.cpp
--- a/Engine/Source/Developer/ShaderFormatOpenGL/Private/ShaderFormatOpenGL.cpp
+++ b/Engine/Source/Developer/ShaderFormatOpenGL/Private/ShaderFormatOpenGL.cpp
@@ -4,6 +4,7 @@
 #include "ShaderFormatOpenGL.h"
 
 #include "ShaderCompilerCommon.h"
+#include "ShaderPreprocessTypes.h"
 #include "HAL/FileManager.h"
 #include "Modules/ModuleManager.h"
 #include "Interfaces/IShaderFormat.h"
@@ -19,9 +20,10 @@ extern bool PreprocessOpenGLShader(
 	const FShaderCompilerEnvironment& Environment,
 	FShaderPreprocessOutput& Output,
 	GLSLVersion Version);
+
 extern void CompileOpenGLShader(
 	const FShaderCompilerInput& Input,
-	const FShaderPreprocessOutput& PreprocessOutput,
+	const FString& InPreprocessedSource,
 	FShaderCompilerOutput& Output,
 	const FString& WorkingDirectory,
 	GLSLVersion Version);
@@ -98,7 +100,7 @@ public:
 	virtual void CompilePreprocessedShader(const FShaderCompilerInput& Input, const FShaderPreprocessOutput& PreprocessOutput, FShaderCompilerOutput& Output, const FString& WorkingDirectory) const override
 	{
 		CheckFormat(Input.ShaderFormat);
-		CompileOpenGLShader(Input, PreprocessOutput, Output, WorkingDirectory, TranslateFormatNameToEnum(Input.ShaderFormat));		
+		CompileOpenGLShader(Input, PreprocessOutput.GetSource(), Output, WorkingDirectory, TranslateFormatNameToEnum(Input.ShaderFormat));		
 	}
 
 	virtual const TCHAR* GetPlatformIncludeDirectory() const override
diff --git a/Engine/Source/Developer/VulkanShaderFormat/Private/VulkanShaderCompiler.cpp b/Engine/Source/Developer/VulkanShaderFormat/Private/VulkanShaderCompiler.cpp
--- a/Engine/Source/Developer/VulkanShaderFormat/Private/VulkanShaderCompiler.cpp
+++ b/Engine/Source/Developer/VulkanShaderFormat/Private/VulkanShaderCompiler.cpp
@@ -2231,17 +2231,6 @@ bool PreprocessVulkanShader(const FShaderCompilerInput& Input, const FShaderComp
 		}
 	}
 
-	// Create a _RootShaderParameters and bind it in slot 0 like any other uniform buffer
-	const TCHAR* ConstantBufferType = InternalState.UseRootParametersStructure() ? TEXT("cbuffer") : nullptr;
-
-	const TArrayView<const TCHAR* const> ExtraSRVTypes;
-	const TArrayView<const TCHAR* const> ExtraUAVTypes;
-	if (!PreprocessOutput.ParseAndModify(Input, Environment, ConstantBufferType, ExtraSRVTypes, ExtraUAVTypes, EBindlessParameterMode::Vulkan))
-	{
-		// The FShaderParameterParser will add any relevant errors.
-		return false;
-	}
-
 	CleanupUniformBufferCode(Input.Environment, PreprocessedShaderSource);
 
 	// Process TEXT macro.
@@ -2543,11 +2532,24 @@ static bool CompileShaderGroup(
 	return bSuccess;
 }
 
-void CompileVulkanShader(const FShaderCompilerInput& Input, const FShaderPreprocessOutput& PreprocessOutput, FShaderCompilerOutput& Output, const class FString& WorkingDirectory)
+void CompileVulkanShader(const FShaderCompilerInput& Input, const FString& InPreprocessedSource, FShaderCompilerOutput& Output, const class FString& WorkingDirectory)
 {
 	check(IsVulkanShaderFormat(Input.ShaderFormat));
 
-	FVulkanShaderCompilerInternalState InternalState(Input, PreprocessOutput.GetParameterParser());
+	FString EntryPointName = Input.EntryPointName;
+	FString PreprocessedSource = InPreprocessedSource;
+
+	// Create a _RootShaderParameters and bind it in slot 0 like any other uniform buffer
+	const TCHAR* ConstantBufferType = (Input.Target.GetFrequency() == SF_RayGen && Input.RootParametersStructure != nullptr) ? TEXT("cbuffer") : nullptr;
+
+	FShaderParameterParser ShaderParameterParser(Input.Environment.CompilerFlags, ConstantBufferType, {}, {});
+	if (!ShaderParameterParser.ParseAndModify(Input, Output.Errors, PreprocessedSource, EBindlessParameterMode::Vulkan))
+	{
+		// The FShaderParameterParser will add any relevant errors.
+		return;
+	}
+
+	FVulkanShaderCompilerInternalState InternalState(Input, ShaderParameterParser);
 
 	const EHlslShaderFrequency HlslFrequency = InternalState.GetHlslShaderFrequency();
 	if (HlslFrequency == HSF_InvalidFrequency)
@@ -2560,13 +2562,15 @@ void CompileVulkanShader(const FShaderCompilerInput& Input, const FShaderPreproc
 		return;
 	}
 
-	FString BindlessPreprocessedShaderSource;
 	if (InternalState.bUseBindlessUniformBuffer)
 	{
-		BindlessPreprocessedShaderSource = PreprocessOutput.GetSource();
-		InternalState.AllBindlessUBs = ConvertUBToBindless(BindlessPreprocessedShaderSource);
+		InternalState.AllBindlessUBs = ConvertUBToBindless(PreprocessedSource);
+	}
+
+	if (ShaderParameterParser.DidModifyShader() || InternalState.AllBindlessUBs.Num() > 0)
+	{
+		Output.ModifiedShaderSource = PreprocessedSource;
 	}
-	const FString& PreprocessedShaderSource = InternalState.bUseBindlessUniformBuffer ? BindlessPreprocessedShaderSource : PreprocessOutput.GetSource();
 
 	bool bSuccess = false;
 
@@ -2574,13 +2578,13 @@ void CompileVulkanShader(const FShaderCompilerInput& Input, const FShaderPreproc
 	// HitGroup shaders might have multiple entrypoints that we combine into a single blob
 	if (InternalState.HasMultipleEntryPoints())
 	{
-		bSuccess = CompileShaderGroup(InternalState, PreprocessedShaderSource, Output);
+		bSuccess = CompileShaderGroup(InternalState, PreprocessedSource, Output);
 	}
 	else
 	{
 		// Compile regular shader via ShaderConductor (DXC)
 		VulkanShaderCompilerSerializedOutput SerializedOutput;
-		bSuccess = CompileWithShaderConductor(InternalState, PreprocessedShaderSource, SerializedOutput, Output);
+		bSuccess = CompileWithShaderConductor(InternalState, PreprocessedSource, SerializedOutput, Output);
 
 		if (InternalState.bUseBindlessUniformBuffer)
 		{
@@ -2606,7 +2610,7 @@ void CompileVulkanShader(const FShaderCompilerInput& Input, const FShaderPreproc
 
 	Output.SerializeShaderCodeValidation();
 
-	PreprocessOutput.GetParameterParser().ValidateShaderParameterTypes(Input, InternalState.IsMobileES31(), Output);
+	ShaderParameterParser.ValidateShaderParameterTypes(Input, InternalState.IsMobileES31(), Output);
 	
 	const bool bDirectCompile = FParse::Param(FCommandLine::Get(), TEXT("directcompile"));
 	if (bDirectCompile)
diff --git a/Engine/Source/Developer/VulkanShaderFormat/Private/VulkanShaderFormat.cpp b/Engine/Source/Developer/VulkanShaderFormat/Private/VulkanShaderFormat.cpp
--- a/Engine/Source/Developer/VulkanShaderFormat/Private/VulkanShaderFormat.cpp
+++ b/Engine/Source/Developer/VulkanShaderFormat/Private/VulkanShaderFormat.cpp
@@ -9,20 +9,21 @@
 #include "hlslcc.h"
 #include "ShaderCore.h"
 #include "ShaderCompilerCore.h"
+#include "ShaderPreprocessTypes.h"
 #include "DXCWrapper.h"
 #include "ShaderConductorContext.h"
 #include "RHIShaderFormatDefinitions.inl"
 
 extern bool PreprocessVulkanShader(
-	const struct  FShaderCompilerInput& Input,
-	const struct FShaderCompilerEnvironment& Environment,
-	class FShaderPreprocessOutput& PreprocessOutput);
+	const FShaderCompilerInput& Input,
+	const FShaderCompilerEnvironment& Environment,
+	FShaderPreprocessOutput& PreprocessOutput);
 
 extern void CompileVulkanShader(
-	const struct FShaderCompilerInput& Input,
-	const class FShaderPreprocessOutput& PreprocessOutput,
-	struct FShaderCompilerOutput& Output,
-	const class FString& WorkingDirectory);
+	const FShaderCompilerInput& Input,
+	const FString& InPreprocessedSource,
+	FShaderCompilerOutput& Output,
+	const FString& WorkingDirectory);
 
 extern void OutputVulkanDebugData(
 	const FShaderCompilerInput& Input, 
@@ -107,7 +108,7 @@ public:
 	virtual void CompilePreprocessedShader(const FShaderCompilerInput& Input, const FShaderPreprocessOutput& PreprocessOutput, FShaderCompilerOutput& Output,const FString& WorkingDirectory) const override
 	{
 		check(InternalGetVersion(Input.ShaderFormat) >= 0);
-		CompileVulkanShader(Input, PreprocessOutput, Output, WorkingDirectory);
+		CompileVulkanShader(Input, PreprocessOutput.GetSource(), Output, WorkingDirectory);
 	}
 
 	virtual void OutputDebugData(const FShaderCompilerInput& Input, const FShaderPreprocessOutput& PreprocessOutput, const FShaderCompilerOutput& Output) const override
diff --git a/Engine/Source/Runtime/RenderCore/Private/ShaderCore.cpp b/Engine/Source/Runtime/RenderCore/Private/ShaderCore.cpp
--- a/Engine/Source/Runtime/RenderCore/Private/ShaderCore.cpp
+++ b/Engine/Source/Runtime/RenderCore/Private/ShaderCore.cpp
@@ -938,6 +938,7 @@ void FShaderCompilerEnvironment::SerializeCompilationDependencies(FArchive& Ar)
 {
 	Ar << CompileArgs;
 	Ar << CompilerFlags;
+	Ar << RenderTargetOutputFormatsMap;
 	Ar << ResourceTableMap.Resources;
 	Ar << UniformBufferMap;
 	Ar << FullPrecisionInPS;
