VulkanRHI: Fix Vulkan bindless and RT heap prefix issues.

#rnx

diff --git a/Engine/Shaders/Public/Platform/Vulkan/VulkanCommon.ush b/Engine/Shaders/Public/Platform/Vulkan/VulkanCommon.ush
--- a/Engine/Shaders/Public/Platform/Vulkan/VulkanCommon.ush
+++ b/Engine/Shaders/Public/Platform/Vulkan/VulkanCommon.ush
@@ -96,7 +96,7 @@ void ImageInterlockedMaxUInt64(RWTexture2D<UlongType> Dest, uint2 Coord, UlongTy
 
 	static const FHitGroupSystemRootConstants HitGroupSystemRootConstants = VulkanHitGroupSystemParameters.RootConstants;
 
-	ByteAddressBuffer VULKAN_RESOURCE_HEAP(HitGroupSystemBuffers)[];
-	static const ByteAddressBuffer HitGroupSystemIndexBuffer = VULKAN_RESOURCE_HEAP(HitGroupSystemBuffers)[VulkanHitGroupSystemParameters.BindlessHitGroupSystemIndexBuffer];
-	static const ByteAddressBuffer HitGroupSystemVertexBuffer = VULKAN_RESOURCE_HEAP(HitGroupSystemBuffers)[VulkanHitGroupSystemParameters.BindlessHitGroupSystemVertexBuffer];
+	ByteAddressBuffer VULKAN_SRV_HEAP(HitGroupSystemBuffers)[];
+	static const ByteAddressBuffer HitGroupSystemIndexBuffer = VULKAN_SRV_HEAP(HitGroupSystemBuffers)[VulkanHitGroupSystemParameters.BindlessHitGroupSystemIndexBuffer];
+	static const ByteAddressBuffer HitGroupSystemVertexBuffer = VULKAN_SRV_HEAP(HitGroupSystemBuffers)[VulkanHitGroupSystemParameters.BindlessHitGroupSystemVertexBuffer];
 #endif
diff --git a/Engine/Source/Developer/VulkanShaderFormat/Private/VulkanShaderCompiler.cpp b/Engine/Source/Developer/VulkanShaderFormat/Private/VulkanShaderCompiler.cpp
--- a/Engine/Source/Developer/VulkanShaderFormat/Private/VulkanShaderCompiler.cpp
+++ b/Engine/Source/Developer/VulkanShaderFormat/Private/VulkanShaderCompiler.cpp
@@ -1240,8 +1240,10 @@ static void GatherSpirvReflectionBindings(
 
 		// Remove resource heaps from binding arrays
 		MoveBindlessHeaps(OutBindings.SBufferUAVs, FShaderParameterParser::kBindlessUAVArrayPrefix, VulkanBindless::BindlessStorageBufferSet);
+		MoveBindlessHeaps(OutBindings.SBufferUAVs, FShaderParameterParser::kBindlessSRVArrayPrefix, VulkanBindless::BindlessStorageBufferSet);  // try with both prefixes, they were merged earlier
 		MoveBindlessHeaps(OutBindings.TextureSRVs, FShaderParameterParser::kBindlessSRVArrayPrefix, VulkanBindless::BindlessSampledImageSet);
 		MoveBindlessHeaps(OutBindings.TextureUAVs, FShaderParameterParser::kBindlessUAVArrayPrefix, VulkanBindless::BindlessStorageImageSet);
+		MoveBindlessHeaps(OutBindings.TextureUAVs, FShaderParameterParser::kBindlessSRVArrayPrefix, VulkanBindless::BindlessStorageImageSet);  // try with both prefixes, R64 SRV textures are read as storage images
 		MoveBindlessHeaps(OutBindings.TBufferSRVs, FShaderParameterParser::kBindlessSRVArrayPrefix, VulkanBindless::BindlessUniformTexelBufferSet);
 		MoveBindlessHeaps(OutBindings.TBufferUAVs, FShaderParameterParser::kBindlessUAVArrayPrefix, VulkanBindless::BindlessStorageTexelBufferSet);
 		MoveBindlessHeaps(OutBindings.AccelerationStructures, FShaderParameterParser::kBindlessSRVArrayPrefix, VulkanBindless::BindlessAccelerationStructureSet);
