Only check empty SPIR-V module when compilation was reported as successful.

#rnx

diff --git a/Engine/Source/Developer/VulkanShaderFormat/Private/VulkanShaderCompiler.cpp b/Engine/Source/Developer/VulkanShaderFormat/Private/VulkanShaderCompiler.cpp
--- a/Engine/Source/Developer/VulkanShaderFormat/Private/VulkanShaderCompiler.cpp
+++ b/Engine/Source/Developer/VulkanShaderFormat/Private/VulkanShaderCompiler.cpp
@@ -2563,14 +2563,17 @@ void CompileVulkanShader(const FShaderCompilerInput& Input, const FString& InPre
 		}
 
 		// Write out the header and shader source code (except for the extra shaders in hit groups)
-		check(SerializedOutput.Spirv.Data.Num() != 0);
+		checkf(!(bSuccess && SerializedOutput.Spirv.Data.Num() == 0), TEXT("shader compilation was reported as successful but SPIR-V module is empty"));
 		FMemoryWriter Ar(Output.ShaderCode.GetWriteAccess(), true);
 		Ar << SerializedOutput.Header;
 		Ar << SerializedOutput.ShaderResourceTable;
 
 		uint32 SpirvCodeSizeBytes = SerializedOutput.Spirv.GetByteSize();
 		Ar << SpirvCodeSizeBytes;
-		Ar.Serialize((uint8*)SerializedOutput.Spirv.Data.GetData(), SpirvCodeSizeBytes);
+		if (SerializedOutput.Spirv.Data.Num() > 0)
+		{
+			Ar.Serialize((uint8*)SerializedOutput.Spirv.Data.GetData(), SpirvCodeSizeBytes);
+		}
 	}
 #endif // PLATFORM_MAC || PLATFORM_WINDOWS || PLATFORM_LINUX
 	
