[shaders] shader format preprocessing cleanup & refactoring
- move uniform buffer cleanup and dead stripping into ShaderPreprocessor module's PreprocessShader function
- add "required symbols" to compiler input struct to specify additional symbols to keep during minification aside from those specified by the entrypoint; modify API such that both an entry point string and additional symbols can be specified (to avoid each backend needing to manually parse the compound RT entry point string)
- make use of ModifyShaderCompilerInput in all backends to set additional defines and required symbols on input struct up front; only use the AdditionalDefines map in cases where it's actually necessary
- remove the various per-platform defines for enabling minifier, no longer required now that this has been rolled out for all backends
- fix SCW directcompile mode; this had rotted due to pieces of the FShaderCompilerEnvironment having been added that weren't explicitly serialized to either cmdline or in the shader source. this now serializes as a base64 string written inside the USF containing all portions of the environment required for compilation (using the same serialization function as is used to write/read the SCW input file)
- use a debug flag for indicating we're in "direct compile" mode and should load the debug USF off disk, rather than the poorly named "bSkipPreprocessedCache" (this name is both inaccurate and also confusing due to the addition of the preprocessed job cache)
- modify platform "force wave32" mechanism to use a pragma directive to set a compiler define, instead of doing string replacement in the preprocessed source
- add a view version of the RT entrypoint parsing to use in preprocessing, note that other paths still need to construct fstrings due to further manipulation so keeping the FString path around too
- clean up backends manually checking the "directcompile" cmdline arg

diff --git a/Engine/Plugins/FX/Niagara/Source/NiagaraEditor/Private/NiagaraCompiler.cpp b/Engine/Plugins/FX/Niagara/Source/NiagaraEditor/Private/NiagaraCompiler.cpp
--- a/Engine/Plugins/FX/Niagara/Source/NiagaraEditor/Private/NiagaraCompiler.cpp
+++ b/Engine/Plugins/FX/Niagara/Source/NiagaraEditor/Private/NiagaraCompiler.cpp
@@ -5,6 +5,7 @@
 #include "DataDrivenShaderPlatformInfo.h"
 #include "EdGraphSchema_Niagara.h"
 #include "EdGraphUtilities.h"
+#include "Interfaces/IShaderFormat.h"
 #include "INiagaraEditorTypeUtilities.h"
 #include "Misc/FileHelper.h"
 #include "Misc/PathViews.h"
@@ -1809,6 +1810,11 @@ int32 FHlslNiagaraCompiler::CompileScript(const FStringView GroupName, const FNi
 	Input.DebugExtension.Empty();
 	Input.DumpDebugInfoPath.Empty();
 
+	FName VVMFormatName = FName(TEXT("VVM_1_0"));
+	//TODO: This is normally invoked by GlobalBeginCompileShader, which is not called in this path. Should it be?
+	const IShaderFormat* VVMShaderFormat = GetTargetPlatformManagerRef().FindShaderFormat(VVMFormatName);
+	VVMShaderFormat->ModifyShaderCompilerInput(Input);
+
 	if (GShaderCompilingManager->GetDumpShaderDebugInfo() == FShaderCompilingManager::EDumpShaderDebugInfo::Always)
 	{
 		Input.DumpDebugInfoPath = GShaderCompilingManager->CreateShaderDebugInfoPath(Input);
@@ -1879,7 +1885,7 @@ int32 FHlslNiagaraCompiler::CompileScript(const FStringView GroupName, const FNi
 			{
 				TArray<FShaderCommonCompileJobPtr> NewJobs;
 				CompilationJob->ShaderCompileJob = Job;
-				Input.ShaderFormat = FName(TEXT("VVM_1_0"));
+				Input.ShaderFormat = VVMFormatName;
 				if (GNiagaraSkipVectorVMBackendOptimizations != 0)
 				{
 					Input.Environment.CompilerFlags.Add(CFLAG_SkipOptimizations);
diff --git a/Engine/Source/Developer/Apple/MetalShaderFormat/Private/MetalDerivedData.cpp b/Engine/Source/Developer/Apple/MetalShaderFormat/Private/MetalDerivedData.cpp
--- a/Engine/Source/Developer/Apple/MetalShaderFormat/Private/MetalDerivedData.cpp
+++ b/Engine/Source/Developer/Apple/MetalShaderFormat/Private/MetalDerivedData.cpp
@@ -1399,9 +1399,8 @@ void DoCompileMetalShader(
 	}
 	else
 	{
-		// Log errors on failed compilation in this backend only when -directcompile is specified.
-		const bool bDirectCompile = FParse::Param(FCommandLine::Get(), TEXT("directcompile"));
-		if (bDirectCompile)
+		// Log errors on failed compilation in this backend only when compiling from a debug dump USF.
+		if (EnumHasAnyFlags(Input.DebugInfoFlags, EShaderDebugInfoFlags::CompileFromDebugUSF))
 		{
 			for (const FShaderCompilerError& Error : Output.Errors)
 			{
diff --git a/Engine/Source/Developer/Apple/MetalShaderFormat/Private/MetalShaderCompiler.cpp b/Engine/Source/Developer/Apple/MetalShaderFormat/Private/MetalShaderCompiler.cpp
--- a/Engine/Source/Developer/Apple/MetalShaderFormat/Private/MetalShaderCompiler.cpp
+++ b/Engine/Source/Developer/Apple/MetalShaderFormat/Private/MetalShaderCompiler.cpp
@@ -763,118 +763,7 @@ bool PreprocessMetalShader(const FShaderCompilerInput& Input, const FShaderCompi
 		return false;
 	}
 
-	PRAGMA_DISABLE_DEPRECATION_WARNINGS		// FShaderCompilerDefinitions will be made internal in the future, marked deprecated until then
-
-	FShaderCompilerDefinitions AdditionalDefines;
-
-	// Work out which standard we need, this is dependent on the shader platform.
-	// TODO: Read from toolchain class
-	const bool bIsMobile = FMetalCompilerToolchain::Get()->IsMobile((EShaderPlatform) Input.Target.Platform);
-	if (bIsMobile)
-	{
-		AdditionalDefines.SetDefine(TEXT("IOS"), 1);
-	}
-	else
-	{
-		AdditionalDefines.SetDefine(TEXT("MAC"), 1);
-	}
-	
-	AdditionalDefines.SetDefine(TEXT("COMPILER_METAL"), 1);
-	
-	if (Input.ShaderFormat == NAME_SF_METAL || Input.ShaderFormat == NAME_SF_METAL_TVOS)
-	{
-		AdditionalDefines.SetDefine(TEXT("METAL_PROFILE"), 1);
-	}
-	else if (Input.ShaderFormat == NAME_SF_METAL_SIM)
-	{
-		AdditionalDefines.SetDefine(TEXT("METAL_PROFILE"), 1);
-	}
-	else if (Input.ShaderFormat == NAME_SF_METAL_MRT || Input.ShaderFormat == NAME_SF_METAL_MRT_TVOS)
-	{
-		AdditionalDefines.SetDefine(TEXT("METAL_MRT_PROFILE"), 1);
-	}
-	else if (Input.ShaderFormat == NAME_SF_METAL_MACES3_1)
-	{
-		AdditionalDefines.SetDefine(TEXT("METAL_PROFILE"), 1);
-	}
-	else if (Input.ShaderFormat == NAME_SF_METAL_SM5)
-	{
-		AdditionalDefines.SetDefine(TEXT("METAL_SM5_PROFILE"), 1);
-		AdditionalDefines.SetDefine(TEXT("USING_VERTEX_SHADER_LAYER"), 1);
-	}
-    else if (Input.ShaderFormat == NAME_SF_METAL_SM6)
-    {
-        AdditionalDefines.SetDefine(TEXT("METAL_SM6_PROFILE"), 1);
-        AdditionalDefines.SetDefine(TEXT("USING_VERTEX_SHADER_LAYER"), 1);
-    }
-	else if (Input.ShaderFormat == NAME_SF_METAL_MRT_MAC)
-	{
-		AdditionalDefines.SetDefine(TEXT("METAL_MRT_PROFILE"), 1);
-	}
-	else
-	{
-		PreprocessOutput.LogError(FString::Printf(TEXT("Invalid shader format '%s' passed to compiler."), *Input.ShaderFormat.ToString()));
-		return false;
-	}
-	
-
-	AdditionalDefines.SetDefine(TEXT("COMPILER_HLSLCC"), 2);
-
-	if (Input.Environment.FullPrecisionInPS)
-	{
-		AdditionalDefines.SetDefine(TEXT("FORCE_FLOATS"), (uint32)1);
-	}
-
-	if (Environment.CompilerFlags.Contains(CFLAG_AvoidFlowControl) 
-		|| Environment.CompilerFlags.Contains(CFLAG_PreferFlowControl))
-	{
-		AdditionalDefines.SetDefine(TEXT("COMPILER_SUPPORTS_ATTRIBUTES"), (uint32)0);
-	}
-	else
-	{
-		AdditionalDefines.SetDefine(TEXT("COMPILER_SUPPORTS_ATTRIBUTES"), (uint32)1);
-	}
-
-	bool bUsesInlineRayTracing = Environment.CompilerFlags.Contains(CFLAG_InlineRayTracing);
-	if (bUsesInlineRayTracing)
-	{
-		AdditionalDefines.SetDefine(TEXT("PLATFORM_SUPPORTS_INLINE_RAY_TRACING"), 1);
-	}
-
-	AdditionalDefines.SetDefine(TEXT("COMPILER_SUPPORTS_DUAL_SOURCE_BLENDING_SLOT_DECORATION"), (uint32)1);
-
-	PRAGMA_ENABLE_DEPRECATION_WARNINGS
-
-	if (Input.bSkipPreprocessedCache)
-	{
-		if (!FFileHelper::LoadFileToString(PreprocessOutput.EditSource(), *Input.VirtualSourceFilePath))
-		{
-			return false;
-		}
-
-		// Remove const as we are on debug-only mode
-		CrossCompiler::CreateEnvironmentFromResourceTable(PreprocessOutput.EditSource(), (FShaderCompilerEnvironment&)Input.Environment);
-	}
-	else
-	{
-		if (!PreprocessShader(PreprocessOutput, Input, Environment, AdditionalDefines))
-		{
-			// The preprocessing stage will add any relevant errors.
-			return false;
-		}
-	}
-
-	CleanupUniformBufferCode(Input.Environment, PreprocessOutput.EditSource());
-
-	// Run the shader minifier
-	#if UE_METAL_SHADER_COMPILER_ALLOW_DEAD_CODE_REMOVAL
-	if (Input.Environment.CompilerFlags.Contains(CFLAG_RemoveDeadCode))
-	{
-		UE::ShaderCompilerCommon::RemoveDeadCode(PreprocessOutput.EditSource(), Input.EntryPointName, PreprocessOutput.EditErrors());
-	}
-	#endif // UE_METAL_SHADER_COMPILER_ALLOW_DEAD_CODE_REMOVAL
-
-	return true;
+	return PreprocessShader(PreprocessOutput, Input, Environment);
 }
 
 void CompileMetalShader(const FShaderCompilerInput& Input, const FString& InPreprocessedSource, FShaderCompilerOutput& Output)
@@ -1069,9 +958,9 @@ void CompileMetalShader(const FShaderCompilerInput& Input, const FString& InPrep
 		Standard = FString::Printf(TEXT("-std=%s-metal%s"), StandardPlatform, *StandardVersion);
 	}
 
-	bool const bDirectCompile = FParse::Param(FCommandLine::Get(), TEXT("directcompile"));
-	if (bDirectCompile)
+	if (EnumHasAnyFlags(Input.DebugInfoFlags, EShaderDebugInfoFlags::CompileFromDebugUSF))
 	{
+		// force debug output on when compiling a debug dump usf
 		const_cast<FShaderCompilerInput&>(Input).DumpDebugInfoPath = FPaths::GetPath(Input.VirtualSourceFilePath);
 	}
 
@@ -1090,7 +979,7 @@ void CompileMetalShader(const FShaderCompilerInput& Input, const FString& InPrep
 
 
 	FSHAHash GUIDHash;
-	if (!bDirectCompile)
+	if (!EnumHasAnyFlags(Input.DebugInfoFlags, EShaderDebugInfoFlags::CompileFromDebugUSF))
 	{
 		TArray<FString> GUIDFiles;
 		GUIDFiles.Add(FPaths::ConvertRelativePathToFull(TEXT("/Engine/Public/Platform/Metal/MetalCommon.ush")));
diff --git a/Engine/Source/Developer/Apple/MetalShaderFormat/Private/MetalShaderFormat.cpp b/Engine/Source/Developer/Apple/MetalShaderFormat/Private/MetalShaderFormat.cpp
--- a/Engine/Source/Developer/Apple/MetalShaderFormat/Private/MetalShaderFormat.cpp
+++ b/Engine/Source/Developer/Apple/MetalShaderFormat/Private/MetalShaderFormat.cpp
@@ -37,7 +37,7 @@ extern uint64 AppendShader_Metal(class FString const& ArchivePath, const FSHAHas
 extern bool FinalizeLibrary_Metal(class FName const& Format, class FString const& ArchivePath, class FString const& LibraryPath, TSet<uint64> const& Shaders, class FString const& DebugOutputDir);
 
 /** Version for shader format, this becomes part of the DDC key. */
-static const FGuid UE_SHADER_METAL_VER = FGuid("FB0EA082-A9FC-469F-AE8F-03CFE0418EA5");
+static const FGuid UE_SHADER_METAL_VER = FGuid("222AACA0-E61A-490A-A5A3-0F3957CC248F");
 
 class FMetalShaderFormat : public IShaderFormat
 {
@@ -90,10 +90,6 @@ public:
 		Result = HashCombine(Result, GetTypeHash(HLSLCC_VersionMinor));
 		Result = HashCombine(Result, GetTypeHash(UE_SHADER_METAL_VER));
 
-#if UE_METAL_SHADER_COMPILER_ALLOW_DEAD_CODE_REMOVAL
-		Result = HashCombine(Result, 0x75E2FE85);
-#endif // UE_METAL_SHADER_COMPILER_ALLOW_DEAD_CODE_REMOVAL
-
 		return Result;
 	}
 
@@ -338,6 +334,79 @@ public:
 		//Map.Format = Format.GetPlainNameString();
 	}
 
+	virtual void ModifyShaderCompilerInput(FShaderCompilerInput& Input) const override
+	{
+		// Work out which standard we need, this is dependent on the shader platform.
+		// TODO: Read from toolchain class
+		const bool bIsMobile = FMetalCompilerToolchain::Get()->IsMobile((EShaderPlatform)Input.Target.Platform);
+		if (bIsMobile)
+		{
+			Input.Environment.SetDefine(TEXT("IOS"), 1);
+		}
+		else
+		{
+			Input.Environment.SetDefine(TEXT("MAC"), 1);
+		}
+
+		Input.Environment.SetDefine(TEXT("COMPILER_METAL"), 1);
+
+		if (Input.ShaderFormat == NAME_SF_METAL || Input.ShaderFormat == NAME_SF_METAL_TVOS)
+		{
+			Input.Environment.SetDefine(TEXT("METAL_PROFILE"), 1);
+		}
+		else if (Input.ShaderFormat == NAME_SF_METAL_SIM)
+		{
+			Input.Environment.SetDefine(TEXT("METAL_PROFILE"), 1);
+		}
+		else if (Input.ShaderFormat == NAME_SF_METAL_MRT || Input.ShaderFormat == NAME_SF_METAL_MRT_TVOS)
+		{
+			Input.Environment.SetDefine(TEXT("METAL_MRT_PROFILE"), 1);
+		}
+		else if (Input.ShaderFormat == NAME_SF_METAL_MACES3_1)
+		{
+			Input.Environment.SetDefine(TEXT("METAL_PROFILE"), 1);
+		}
+		else if (Input.ShaderFormat == NAME_SF_METAL_SM5)
+		{
+			Input.Environment.SetDefine(TEXT("METAL_SM5_PROFILE"), 1);
+			Input.Environment.SetDefine(TEXT("USING_VERTEX_SHADER_LAYER"), 1);
+		}
+		else if (Input.ShaderFormat == NAME_SF_METAL_SM6)
+		{
+			Input.Environment.SetDefine(TEXT("METAL_SM6_PROFILE"), 1);
+			Input.Environment.SetDefine(TEXT("USING_VERTEX_SHADER_LAYER"), 1);
+		}
+		else if (Input.ShaderFormat == NAME_SF_METAL_MRT_MAC)
+		{
+			Input.Environment.SetDefine(TEXT("METAL_MRT_PROFILE"), 1);
+		}
+
+		Input.Environment.SetDefine(TEXT("COMPILER_HLSLCC"), 2);
+
+		if (Input.Environment.FullPrecisionInPS || (IsValidRef(Input.SharedEnvironment) && Input.SharedEnvironment->FullPrecisionInPS))
+		{
+			Input.Environment.SetDefine(TEXT("FORCE_FLOATS"), (uint32)1);
+		}
+
+		if (Input.Environment.CompilerFlags.Contains(CFLAG_AvoidFlowControl)
+			|| Input.Environment.CompilerFlags.Contains(CFLAG_PreferFlowControl))
+		{
+			Input.Environment.SetDefine(TEXT("COMPILER_SUPPORTS_ATTRIBUTES"), (uint32)0);
+		}
+		else
+		{
+			Input.Environment.SetDefine(TEXT("COMPILER_SUPPORTS_ATTRIBUTES"), (uint32)1);
+		}
+
+		bool bUsesInlineRayTracing = Input.Environment.CompilerFlags.Contains(CFLAG_InlineRayTracing);
+		if (bUsesInlineRayTracing)
+		{
+			Input.Environment.SetDefine(TEXT("PLATFORM_SUPPORTS_INLINE_RAY_TRACING"), 1);
+		}
+
+		Input.Environment.SetDefine(TEXT("COMPILER_SUPPORTS_DUAL_SOURCE_BLENDING_SLOT_DECORATION"), (uint32)1);
+	}
+	
 	virtual bool CanCompileBinaryShaders() const override final
 	{
 #if PLATFORM_MAC
diff --git a/Engine/Source/Developer/Apple/MetalShaderFormat/Private/MetalShaderFormat.h b/Engine/Source/Developer/Apple/MetalShaderFormat/Private/MetalShaderFormat.h
--- a/Engine/Source/Developer/Apple/MetalShaderFormat/Private/MetalShaderFormat.h
+++ b/Engine/Source/Developer/Apple/MetalShaderFormat/Private/MetalShaderFormat.h
@@ -11,11 +11,6 @@
 #include "RHIFeatureLevel.h"
 #include "RHIShaderPlatform.h"
 
-// Controls whether r.Shaders.RemoveDeadCode should be honored
-#ifndef UE_METAL_SHADER_COMPILER_ALLOW_DEAD_CODE_REMOVAL
-#define UE_METAL_SHADER_COMPILER_ALLOW_DEAD_CODE_REMOVAL 1
-#endif // UE_METAL_SHADER_COMPILER_ALLOW_DEAD_CODE_REMOVAL
-
 // IOS and TVOS use the mobile toolchain.
 enum EAppleSDKType
 {
diff --git a/Engine/Source/Developer/ShaderCompilerCommon/Private/ShaderCompilerCommon.cpp b/Engine/Source/Developer/ShaderCompilerCommon/Private/ShaderCompilerCommon.cpp
--- a/Engine/Source/Developer/ShaderCompilerCommon/Private/ShaderCompilerCommon.cpp
+++ b/Engine/Source/Developer/ShaderCompilerCommon/Private/ShaderCompilerCommon.cpp
@@ -2,12 +2,15 @@
 
 #include "ShaderCompilerCommon.h"
 #include "ShaderParameterParser.h"
+#include "Misc/Base64.h"
 #include "Misc/FileHelper.h"
 #include "Misc/Paths.h"
 #include "Misc/PathViews.h"
 #include "Modules/ModuleManager.h"
 #include "HlslccDefinitions.h"
 #include "HAL/FileManager.h"
+#include "Serialization/MemoryReader.h"
+#include "Serialization/MemoryWriter.h"
 #include "String/RemoveFrom.h"
 #include "ShaderPreprocessTypes.h"
 #include "ShaderSymbolExport.h"
@@ -422,29 +425,33 @@ bool UE::ShaderCompilerCommon::ValidatePackedResourceCounts(FShaderCompilerOutpu
 	return Output.bSucceeded;
 }
 
-void UE::ShaderCompilerCommon::ParseRayTracingEntryPoint(const FString& Input, FString& OutMain, FString& OutAnyHit, FString& OutIntersection)
+void UE::ShaderCompilerCommon::ParseRayTracingEntryPoint(const FStringView& Input, FStringView& OutMain, FStringView& OutAnyHit, FStringView& OutIntersection)
 {
-	auto ParseEntry = [&Input](const TCHAR* Marker)
+	auto ParseEntry = [&Input](const FStringView& Marker)
 	{
-		FString Result;
-		const int32 BeginIndex = Input.Find(Marker, ESearchCase::IgnoreCase, ESearchDir::FromStart);
+		FStringView Result;
+		const int32 BeginIndex = UE::String::FindFirst(Input, Marker, ESearchCase::IgnoreCase);
 		if (BeginIndex != INDEX_NONE)
 		{
-			int32 EndIndex = Input.Find(TEXT(" "), ESearchCase::IgnoreCase, ESearchDir::FromStart, BeginIndex);
+			int32 EndIndex = UE::String::FindFirst(Input.Mid(BeginIndex), TEXTVIEW(" "), ESearchCase::IgnoreCase);
 			if (EndIndex == INDEX_NONE)
 			{
 				EndIndex = Input.Len() + 1;
 			}
-			const int32 MarkerLen = FCString::Strlen(Marker);
+			else
+			{
+				EndIndex += BeginIndex;
+			}
+			const int32 MarkerLen = Marker.Len();
 			const int32 Count = EndIndex - BeginIndex;
 			Result = Input.Mid(BeginIndex + MarkerLen, Count - MarkerLen);
 		}
 		return Result;
 	};
 
-	OutMain = ParseEntry(TEXT("closesthit="));
-	OutAnyHit = ParseEntry(TEXT("anyhit="));
-	OutIntersection = ParseEntry(TEXT("intersection="));
+	OutMain = ParseEntry(TEXTVIEW("closesthit="));
+	OutAnyHit = ParseEntry(TEXTVIEW("anyhit="));
+	OutIntersection = ParseEntry(TEXTVIEW("intersection="));
 
 	// If complex hit group entry is not specified, assume a single verbatim entry point
 	if (OutMain.IsEmpty() && OutAnyHit.IsEmpty() && OutIntersection.IsEmpty())
@@ -453,6 +460,19 @@ void UE::ShaderCompilerCommon::ParseRayTracingEntryPoint(const FString& Input, F
 	}
 }
 
+void UE::ShaderCompilerCommon::ParseRayTracingEntryPoint(const FString& Input, FString& OutMain, FString& OutAnyHit, FString& OutIntersection)
+{
+	FStringView OutMainView;
+	FStringView OutAnyHitView;
+	FStringView OutIntersectionView;
+	ParseRayTracingEntryPoint(Input, OutMainView, OutAnyHitView, OutIntersectionView);
+
+	OutMain = OutMainView;
+	OutAnyHit = OutAnyHitView;
+	OutIntersection = OutIntersectionView;
+}
+
+
 bool UE::ShaderCompilerCommon::RemoveDeadCode(FString& InOutPreprocessedShaderSource, TConstArrayView<FStringView> RequiredSymbols, TArray<FShaderCompilerError>& OutErrors)
 {
 	TRACE_CPUPROFILER_EVENT_SCOPE(RemoveDeadCode);
@@ -482,12 +502,17 @@ bool UE::ShaderCompilerCommon::RemoveDeadCode(FString& InOutPreprocessedShaderSo
 }
 
 bool UE::ShaderCompilerCommon::RemoveDeadCode(FString& InOutPreprocessedShaderSource, const FString& EntryPoint, TArray<FShaderCompilerError>& OutErrors)
+{
+	return UE::ShaderCompilerCommon::RemoveDeadCode(InOutPreprocessedShaderSource, EntryPoint, {}, OutErrors);
+}
+
+bool UE::ShaderCompilerCommon::RemoveDeadCode(FString& InOutPreprocessedShaderSource, const FString& EntryPoint, TConstArrayView<FStringView> InRequiredSymbols, TArray<FShaderCompilerError>& OutErrors)
 {
 	TArray<FStringView> RequiredSymbols;
 
-	FString EntryMain;
-	FString EntryAnyHit;
-	FString EntryIntersection;
+	FStringView EntryMain;
+	FStringView EntryAnyHit;
+	FStringView EntryIntersection;
 	UE::ShaderCompilerCommon::ParseRayTracingEntryPoint(EntryPoint, EntryMain, EntryAnyHit, EntryIntersection);
 
 	RequiredSymbols.Add(EntryMain);
@@ -502,6 +527,11 @@ bool UE::ShaderCompilerCommon::RemoveDeadCode(FString& InOutPreprocessedShaderSo
 		RequiredSymbols.Add(EntryIntersection);
 	}
 
+	for (FStringView Symbol : InRequiredSymbols)
+	{
+		RequiredSymbols.Add(Symbol);
+	}
+
 	return UE::ShaderCompilerCommon::RemoveDeadCode(InOutPreprocessedShaderSource, RequiredSymbols, OutErrors);
 }
 
@@ -1517,9 +1547,6 @@ FString CreateShaderCompilerWorkerDirectCommandLine(const FShaderCompilerInput&
 	Text += TEXT(" ");
 	Text += Input.DumpDebugInfoPath / Input.GetSourceFilename();
 
-	Text += TEXT(" -cflags=");
-	Text += FString::Printf(TEXT("%llu"), Input.Environment.CompilerFlags.GetData());
-
 	// When we're running in directcompile mode, we don't to spam the crash reporter
 	Text += TEXT(" -nocrashreports");
 	return Text;
@@ -1878,6 +1905,39 @@ namespace UE::ShaderCompilerCommon
 		}
 	}
 
+	static const TCHAR* Base64EnvBegin = TEXT("#if 0 /* BASE64_ENV */\n");
+	static const int32 Base64EnvBeginLen = FCString::Strlen(Base64EnvBegin);
+	static const TCHAR* Base64EnvEnd = TEXT("\n#endif /* BASE64_ENV */\n");
+	
+	FString SerializeEnvironmentToBase64(const FShaderCompilerEnvironment& Env)
+	{
+		TArray<uint8> Serialized;
+		FMemoryWriter Ar(Serialized);
+		const_cast<FShaderCompilerEnvironment&>(Env).SerializeCompilationDependencies(Ar);
+		return FString::Printf(TEXT("%s%s%s"), Base64EnvBegin, *FBase64::Encode(Serialized), Base64EnvEnd);
+	}
+
+	void SerializeEnvironmentFromBase64(FShaderCompilerEnvironment& Env, const FString& DebugShaderSource)
+	{
+		int32 BeginIndex = DebugShaderSource.Find(Base64EnvBegin, ESearchCase::CaseSensitive);
+		if (BeginIndex == INDEX_NONE)
+		{
+			return;
+		}
+		int32 EndIndex = DebugShaderSource.Find(Base64EnvEnd, ESearchCase::CaseSensitive, ESearchDir::FromStart, BeginIndex);
+		if (EndIndex == INDEX_NONE)
+		{
+			return;
+		}
+
+		FString Base64Encoded = DebugShaderSource.Left(EndIndex).Mid(BeginIndex + Base64EnvBeginLen);
+
+		TArray<uint8> Serialized;
+		FBase64::Decode(Base64Encoded, Serialized);
+		FMemoryReader Ar(Serialized);
+		Env.SerializeCompilationDependencies(Ar);
+	}
+
 	FString GetDebugShaderContents(const FShaderCompilerInput& Input, const FString& PreprocessedSource, const FDebugShaderDataOptions& Options)
 	{
 		FString Contents = Options.AppendPreSource ? Options.AppendPreSource() : FString();
@@ -1886,8 +1946,15 @@ namespace UE::ShaderCompilerCommon
 		{
 			Contents += Options.AppendPostSource();
 		}
+		// If preprocessed cache is enabled, debug dump occurs in the cook process rather than the workers, and
+		// in that case the env in Input.Environment has not been merged with the shared env. Do so here.
+		FShaderCompilerEnvironment MergedEnvironment(Input.Environment);
+		if (Input.bCachePreprocessed && IsValidRef(Input.SharedEnvironment))
+		{
+			MergedEnvironment.Merge(*Input.SharedEnvironment);
+		}
 		Contents += TEXT("\n");
-		Contents += CrossCompiler::CreateResourceTableFromEnvironment(Input.Environment);
+		Contents += SerializeEnvironmentToBase64(MergedEnvironment);
 		Contents += TEXT("#if 0 /*DIRECT COMPILE*/\n");
 		Contents += CreateShaderCompilerWorkerDirectCommandLine(Input);
 		Contents += TEXT("\n#endif /*DIRECT COMPILE*/\n");
diff --git a/Engine/Source/Developer/ShaderCompilerCommon/Public/ShaderCompilerCommon.h b/Engine/Source/Developer/ShaderCompilerCommon/Public/ShaderCompilerCommon.h
--- a/Engine/Source/Developer/ShaderCompilerCommon/Public/ShaderCompilerCommon.h
+++ b/Engine/Source/Developer/ShaderCompilerCommon/Public/ShaderCompilerCommon.h
@@ -69,6 +69,8 @@ namespace UE::ShaderCompilerCommon
 
 	extern SHADERCOMPILERCOMMON_API bool                 ValidatePackedResourceCounts(FShaderCompilerOutput& Output, const FShaderCodePackedResourceCounts& PackedResourceCounts);
 
+	extern SHADERCOMPILERCOMMON_API void ParseRayTracingEntryPoint(const FString& Input, FString& OutMain, FString& OutAnyHit, FString& OutIntersection);
+
 	/*
 	* Parses ray tracing shader entry point specification string in one of the following formats:
 	* 1) Verbatim single entry point name, e.g. "MainRGS"
@@ -79,7 +81,7 @@ namespace UE::ShaderCompilerCommon
 	*      d) "closesthit=MainCHS intersection=MainIS"
 	*    NOTE: closesthit attribute must always be provided for complex hit group entry points
 	*/
-	extern SHADERCOMPILERCOMMON_API void ParseRayTracingEntryPoint(const FString& Input, FString& OutMain, FString& OutAnyHit, FString& OutIntersection);
+	extern SHADERCOMPILERCOMMON_API void ParseRayTracingEntryPoint(const FStringView& Input, FStringView& OutMain, FStringView& OutAnyHit, FStringView& OutIntersection);
 
 	/**
 	* Rewrites a fully preprocessed shader source code, removing any functions or structs that are not reachable from a given entry point or list of symbols.
@@ -88,6 +90,7 @@ namespace UE::ShaderCompilerCommon
 	* InOutPreprocessedShaderSource is replaced with the rewritten code on success and is kept intact on failure.
 	*/
 	extern SHADERCOMPILERCOMMON_API bool RemoveDeadCode(FString& InOutPreprocessedShaderSource, const FString& EntryPoint, TArray<FShaderCompilerError>& OutErrors);
+	extern SHADERCOMPILERCOMMON_API bool RemoveDeadCode(FString& InOutPreprocessedShaderSource, const FString& EntryPoint, TConstArrayView<FStringView> RequiredSymbols, TArray<FShaderCompilerError>& OutErrors);
 	extern SHADERCOMPILERCOMMON_API bool RemoveDeadCode(FString& InOutPreprocessedShaderSource, TConstArrayView<FStringView> RequiredSymbols, TArray<FShaderCompilerError>& OutErrors);
 
 	struct FDebugShaderDataOptions
@@ -149,6 +152,9 @@ namespace UE::ShaderCompilerCommon
 		const FShaderCompilerOutput& Output,
 		const UE::ShaderCompilerCommon::FDebugShaderDataOptions& Options = FDebugShaderDataOptions());
 
+	extern SHADERCOMPILERCOMMON_API void SerializeEnvironmentFromBase64(FShaderCompilerEnvironment& Env, const FString& DebugUSF);
+	extern SHADERCOMPILERCOMMON_API FString SerializeEnvironmentToBase64(const FShaderCompilerEnvironment& Env);
+
 	/*
 	 * Constructs the modified preprocessed source that would be dumped to a .usf file via DumpDebugShaderData, including the following additions:
 	 * - resource table debug information
@@ -395,7 +401,9 @@ extern SHADERCOMPILERCOMMON_API void CompileOfflineMali(const FShaderCompilerInp
 // Cross compiler support/common functionality
 namespace CrossCompiler
 {
+	UE_DEPRECATED(5.4, "CreateResourceTableFromEnvironment is no longer used; serializing environment for SCW directcompile mode now uses a base64-encoded string containing all environment compilation dependencies.")
 	extern SHADERCOMPILERCOMMON_API FString CreateResourceTableFromEnvironment(const FShaderCompilerEnvironment& Environment);
+	UE_DEPRECATED(5.4, "CreateEnvironmentFromResourceTable is no longer used; serializing environment for SCW directcompile mode now uses a base64-encoded string containing all environment compilation dependencies.")
 	extern SHADERCOMPILERCOMMON_API void CreateEnvironmentFromResourceTable(const FString& String, FShaderCompilerEnvironment& OutEnvironment);
 
 	extern SHADERCOMPILERCOMMON_API void ParseHlslccError(TArray<FShaderCompilerError>& OutErrors, const FString& InLine, bool bUseAbsolutePaths = false);
diff --git a/Engine/Source/Developer/ShaderFormatOpenGL/Private/OpenGLShaderCompiler.cpp b/Engine/Source/Developer/ShaderFormatOpenGL/Private/OpenGLShaderCompiler.cpp
--- a/Engine/Source/Developer/ShaderFormatOpenGL/Private/OpenGLShaderCompiler.cpp
+++ b/Engine/Source/Developer/ShaderFormatOpenGL/Private/OpenGLShaderCompiler.cpp
@@ -1067,30 +1067,9 @@ static EHlslCompileTarget GetCompileTarget(GLSLVersion Version)
 
 PRAGMA_DISABLE_DEPRECATION_WARNINGS		// FShaderCompilerDefinitions will be made internal in the future, marked deprecated until then
 
-static void SetupDefines(GLSLVersion Version, FShaderCompilerDefinitions& AdditionalDefines)
-{
-	switch (Version)
-	{
-		case GLSL_ES3_1_ANDROID:
-			AdditionalDefines.SetDefine(TEXT("COMPILER_GLSL_ES3_1"), 1);
-			AdditionalDefines.SetDefine(TEXT("ES3_1_PROFILE"), 1);
-			break;
-
-		case GLSL_150_ES3_1:
-			AdditionalDefines.SetDefine(TEXT("COMPILER_GLSL"), 1);
-			AdditionalDefines.SetDefine(TEXT("ES3_1_PROFILE"), 1);
-			AdditionalDefines.SetDefine(TEXT("row_major"), TEXT(""));
-			break;
-
-		default:
-			check(0);
-	}
-	AdditionalDefines.SetDefine(TEXT("OPENGL_PROFILE"), 1);
-}
-
 void FOpenGLFrontend::SetupPerVersionCompilationEnvironment(GLSLVersion Version, FShaderCompilerDefinitions& AdditionalDefines, EHlslCompileTarget& HlslCompilerTarget)
 {
-	SetupDefines(Version, AdditionalDefines);
+	// this function is deprecated; SetupDefines is removed and all per-version defines are now set in ModifyShaderCompilerInput on the IShaderFormat implementation.
 	HlslCompilerTarget = GetCompileTarget(Version);
 }
 
@@ -3184,7 +3163,7 @@ static bool CompileToGlslWithShaderConductor(
 			{
 				// handle both the case where identifier for TextureExternal declaration immediately precedes a ; and has whitespace separating the two
 #if PLATFORM_WINDOWS
-				if (swscanf_s(&PreprocessedShader[Pos], TEXT("TextureExternal %ls %ls"), TextureExternalName, 256, NextToken, 4))
+				if (swscanf_s(&PreprocessedShader[Pos], TEXT("TextureExternal %ls %ls"), TextureExternalName, 257, NextToken, 2))
 #elif PLATFORM_MAC
 				if (sscanf(TCHAR_TO_ANSI(&PreprocessedShader[Pos]), "TextureExternal %256s %1s", TextureExternalName, NextToken))
 #else // PLATFORM_LINUX
@@ -3269,7 +3248,7 @@ static bool CompileToGlslWithShaderConductor(
 
 #endif // DXC_SUPPORTED
 
-static bool ShouldUseDXC(FShaderCompilerFlags Flags)
+bool ShouldUseDXC(FShaderCompilerFlags Flags)
 {
 #if DXC_SUPPORTED
 	return Flags.Contains(CFLAG_ForceDXC);
@@ -3278,56 +3257,6 @@ static bool ShouldUseDXC(FShaderCompilerFlags Flags)
 #endif
 }
 
-bool PreprocessOpenGLShader(const FShaderCompilerInput& Input, const FShaderCompilerEnvironment& Environment, FShaderPreprocessOutput& PreprocessOutput, GLSLVersion Version)
-{
-	PRAGMA_DISABLE_DEPRECATION_WARNINGS // FShaderCompilerDefinitions will be made internal in the future, marked deprecated until then
-	
-	FShaderCompilerDefinitions AdditionalDefines;
-	SetupDefines(Version, AdditionalDefines);
-
-	const bool bUseSC = ShouldUseDXC(Environment.CompilerFlags);
-	AdditionalDefines.SetDefine(TEXT("COMPILER_HLSLCC"), bUseSC ? 2 : 1);
-	AdditionalDefines.SetDefine(TEXT("COMPILER_SUPPORTS_ATTRIBUTES"), (uint32)1);
-
-	if (Input.Environment.FullPrecisionInPS)
-	{
-		AdditionalDefines.SetDefine(TEXT("FORCE_FLOATS"), (uint32)1);
-	}
-
-	PRAGMA_ENABLE_DEPRECATION_WARNINGS
-	
-	if (Input.bSkipPreprocessedCache)
-	{
-		if (!FFileHelper::LoadFileToString(PreprocessOutput.EditSource(), *Input.VirtualSourceFilePath))
-		{
-			return false;
-		}
-
-		// Remove const as we are on debug-only mode
-		CrossCompiler::CreateEnvironmentFromResourceTable(PreprocessOutput.GetSource(), const_cast<FShaderCompilerEnvironment&>(Environment));
-	}
-	else
-	{
-		if (!PreprocessShader(PreprocessOutput, Input, Environment, AdditionalDefines))
-		{
-			// The preprocessing stage will add any relevant errors.
-			return false;
-		}
-	}
-
-	// This requires removing the HLSLCC_NoPreprocess flag later on!
-	CleanupUniformBufferCode(Input.Environment, PreprocessOutput.EditSource());
-
-	// Run the experimental shader minifier
-#if UE_OPENGL_SHADER_COMPILER_ALLOW_DEAD_CODE_REMOVAL
-	if (Input.Environment.CompilerFlags.Contains(CFLAG_RemoveDeadCode))
-	{
-		UE::ShaderCompilerCommon::RemoveDeadCode(PreprocessOutput.EditSource(), Input.EntryPointName, PreprocessOutput.EditErrors());
-	}
-#endif // UE_OPENGL_SHADER_COMPILER_ALLOW_DEAD_CODE_REMOVAL
-	return true;
-}
-
 /**
  * Compile a shader for OpenGL on Windows.
  * @param Input - The input shader code and environment.
@@ -3431,10 +3360,10 @@ void CompileOpenGLShader(const FShaderCompilerInput& Input, const FString& InPre
 		}
 	}
 
-	static const bool bDirectCompile = FParse::Param(FCommandLine::Get(), TEXT("directcompile"));
 	if (bCompilationSucceeded)
 	{
-		if (bDirectCompile)
+		// print resulting glsl to debug output if compiling from a debug dumped usf file
+		if (EnumHasAnyFlags(Input.DebugInfoFlags, EShaderDebugInfoFlags::CompileFromDebugUSF))
 		{
 			FPlatformMisc::LowLevelOutputDebugStringf(TEXT("%s\n"), ANSI_TO_TCHAR(GlslShaderSource));
 		}
@@ -3460,7 +3389,7 @@ void CompileOpenGLShader(const FShaderCompilerInput& Input, const FString& InPre
 	}
 	else if (!bUseSC)
 	{
-		const bool bUseAbsolutePaths = bDirectCompile;
+		const bool bUseAbsolutePaths = EnumHasAnyFlags(Input.DebugInfoFlags, EShaderDebugInfoFlags::CompileFromDebugUSF);
 
 		const FString Tmp = ANSI_TO_TCHAR(ErrorLog);
 		TArray<FString> ErrorLines;
@@ -3489,7 +3418,7 @@ void CompileOpenGLShader(const FShaderCompilerInput& Input, const FString& InPre
 void FOpenGLFrontend::CompileShader(const FShaderCompilerInput& Input, FShaderCompilerOutput& Output, const FString& WorkingDirectory, GLSLVersion Version)
 {
 	FShaderPreprocessOutput PreprocessOutput;
-	PreprocessOpenGLShader(Input, Input.Environment, PreprocessOutput, Version);
+	PreprocessShader(PreprocessOutput, Input, Input.Environment);
 	CompileOpenGLShader(Input, PreprocessOutput.GetSource(), Output, WorkingDirectory, Version);
 }
 
diff --git a/Engine/Source/Developer/ShaderFormatOpenGL/Private/ShaderFormatOpenGL.cpp b/Engine/Source/Developer/ShaderFormatOpenGL/Private/ShaderFormatOpenGL.cpp
--- a/Engine/Source/Developer/ShaderFormatOpenGL/Private/ShaderFormatOpenGL.cpp
+++ b/Engine/Source/Developer/ShaderFormatOpenGL/Private/ShaderFormatOpenGL.cpp
@@ -4,6 +4,7 @@
 #include "ShaderFormatOpenGL.h"
 
 #include "ShaderCompilerCommon.h"
+#include "ShaderPreprocessor.h"
 #include "ShaderPreprocessTypes.h"
 #include "HAL/FileManager.h"
 #include "Modules/ModuleManager.h"
@@ -15,11 +16,7 @@
 static FName NAME_GLSL_150_ES3_1(TEXT("GLSL_150_ES31"));
 static FName NAME_GLSL_ES3_1_ANDROID(TEXT("GLSL_ES3_1_ANDROID"));
 
-extern bool PreprocessOpenGLShader(
-	const FShaderCompilerInput& Input,
-	const FShaderCompilerEnvironment& Environment,
-	FShaderPreprocessOutput& Output,
-	GLSLVersion Version);
+extern bool ShouldUseDXC(FShaderCompilerFlags Flags);
 
 extern void CompileOpenGLShader(
 	const FShaderCompilerInput& Input,
@@ -29,7 +26,7 @@ extern void CompileOpenGLShader(
 	GLSLVersion Version);
 
 /** Version for shader format, this becomes part of the DDC key. */
-static const FGuid UE_SHADER_GLSL_VER = FGuid("45CB4440-A342-475E-AAD1-4CB8E228D1C9");
+static const FGuid UE_SHADER_GLSL_VER = FGuid("15A0CD3A-20CC-4F15-95CC-04E631068DFE");
 
 class FShaderFormatGLSL : public UE::ShaderCompilerCommon::FBaseShaderFormat 
 {
@@ -53,10 +50,6 @@ public:
 		uint32 Version = GetTypeHash(HLSLCC_VersionMinor);
 		Version = HashCombine(Version, BaseHash);
 
-	#if UE_OPENGL_SHADER_COMPILER_ALLOW_DEAD_CODE_REMOVAL
-		Version = HashCombine(Version, 0x75E2FE85);
-	#endif // UE_OPENGL_SHADER_COMPILER_ALLOW_DEAD_CODE_REMOVAL
-
 		return Version;
 	}
 	virtual void GetSupportedFormats(TArray<FName>& OutFormats) const override
@@ -82,10 +75,41 @@ public:
 		}
 	}
 
+	virtual void ModifyShaderCompilerInput(FShaderCompilerInput& Input) const override
+	{
+		GLSLVersion Version = TranslateFormatNameToEnum(Input.ShaderFormat);
+		switch (Version)
+		{
+		case GLSL_ES3_1_ANDROID:
+			Input.Environment.SetDefine(TEXT("COMPILER_GLSL_ES3_1"), 1);
+			Input.Environment.SetDefine(TEXT("ES3_1_PROFILE"), 1);
+			break;
+
+		case GLSL_150_ES3_1:
+			Input.Environment.SetDefine(TEXT("COMPILER_GLSL"), 1);
+			Input.Environment.SetDefine(TEXT("ES3_1_PROFILE"), 1);
+			Input.Environment.SetDefine(TEXT("row_major"), TEXT(""));
+			break;
+
+		default:
+			check(0);
+		}
+		Input.Environment.SetDefine(TEXT("OPENGL_PROFILE"), 1);
+
+		const bool bUseDXC = ShouldUseDXC(Input.Environment.CompilerFlags);
+		Input.Environment.SetDefine(TEXT("COMPILER_HLSLCC"), bUseDXC ? 2 : 1);
+		Input.Environment.SetDefine(TEXT("COMPILER_SUPPORTS_ATTRIBUTES"), (uint32)1);
+
+		if (Input.Environment.FullPrecisionInPS || (IsValidRef(Input.SharedEnvironment) && Input.SharedEnvironment->FullPrecisionInPS))
+		{
+			Input.Environment.SetDefine(TEXT("FORCE_FLOATS"), (uint32)1);
+		}
+	}
+
 	virtual bool PreprocessShader(const FShaderCompilerInput& Input, const FShaderCompilerEnvironment& Environment, FShaderPreprocessOutput& PreprocessOutput) const override
 	{
 		CheckFormat(Input.ShaderFormat);
-		return PreprocessOpenGLShader(Input, Environment, PreprocessOutput, TranslateFormatNameToEnum(Input.ShaderFormat));
+		return ::PreprocessShader(PreprocessOutput, Input, Environment);
 	}
 
 	virtual void CompilePreprocessedShader(const FShaderCompilerInput& Input, const FShaderPreprocessOutput& PreprocessOutput, FShaderCompilerOutput& Output, const FString& WorkingDirectory) const override
diff --git a/Engine/Source/Developer/ShaderFormatOpenGL/Public/ShaderFormatOpenGL.h b/Engine/Source/Developer/ShaderFormatOpenGL/Public/ShaderFormatOpenGL.h
--- a/Engine/Source/Developer/ShaderFormatOpenGL/Public/ShaderFormatOpenGL.h
+++ b/Engine/Source/Developer/ShaderFormatOpenGL/Public/ShaderFormatOpenGL.h
@@ -17,11 +17,6 @@ class FArchive;
 #include "Windows/WindowsHWrapper.h"
 #endif
 
-// Controls whether r.Shaders.RemoveDeadCode should be honored
-#ifndef UE_OPENGL_SHADER_COMPILER_ALLOW_DEAD_CODE_REMOVAL
-#define UE_OPENGL_SHADER_COMPILER_ALLOW_DEAD_CODE_REMOVAL 1
-#endif // UE_OPENGL_SHADER_COMPILER_ALLOW_DEAD_CODE_REMOVAL
-
 class FShaderCompilerFlags;
 struct FShaderCompilerInput;
 struct FShaderCompilerOutput;
diff --git a/Engine/Source/Developer/ShaderFormatVectorVM/Private/ShaderFormatVectorVM.cpp b/Engine/Source/Developer/ShaderFormatVectorVM/Private/ShaderFormatVectorVM.cpp
--- a/Engine/Source/Developer/ShaderFormatVectorVM/Private/ShaderFormatVectorVM.cpp
+++ b/Engine/Source/Developer/ShaderFormatVectorVM/Private/ShaderFormatVectorVM.cpp
@@ -63,6 +63,17 @@ public:
 		OutFormats.Add(NAME_VVM_1_0);
 	}
 
+	virtual void ModifyShaderCompilerInput(FShaderCompilerInput& Input) const override
+	{
+		Input.Environment.SetDefine(TEXT("COMPILER_HLSLCC"), 1);
+		Input.Environment.SetDefine(TEXT("COMPILER_VECTORVM"), 1);
+		Input.Environment.SetDefine(TEXT("VECTORVM_PROFILE"), 1);
+		Input.Environment.SetDefine(TEXT("FORCE_FLOATS"), (uint32)1);
+
+		// minifier has not been tested on vector VM; it's possible this could be removed to improve deduplication rate
+		Input.Environment.CompilerFlags.Remove(CFLAG_RemoveDeadCode);
+	}
+
 	virtual bool PreprocessShader(const FShaderCompilerInput& Input, const FShaderCompilerEnvironment& Environment, FShaderPreprocessOutput& PreprocessOutput) const override
 	{
 		CheckFormat(Input.ShaderFormat);
diff --git a/Engine/Source/Developer/ShaderFormatVectorVM/Private/VectorVMShaderCompiler.cpp b/Engine/Source/Developer/ShaderFormatVectorVM/Private/VectorVMShaderCompiler.cpp
--- a/Engine/Source/Developer/ShaderFormatVectorVM/Private/VectorVMShaderCompiler.cpp
+++ b/Engine/Source/Developer/ShaderFormatVectorVM/Private/VectorVMShaderCompiler.cpp
@@ -29,35 +29,12 @@ DECLARE_CYCLE_STAT(TEXT("VectorVM - Compiler - CrossCompilerContextRun"), STAT_V
 bool PreprocessVectorVMShader(const FShaderCompilerInput& Input, const FShaderCompilerEnvironment& Environment, FShaderPreprocessOutput& Output)
 {
 	SCOPE_CYCLE_COUNTER(STAT_VectorVM_Compiler_CompileShader_VectorVMPreprocessShader);
-
-	PRAGMA_DISABLE_DEPRECATION_WARNINGS		// FShaderCompilerDefinitions will be made internal in the future, marked deprecated until then
-	FShaderCompilerDefinitions AdditionalDefines;
-	AdditionalDefines.SetDefine(TEXT("COMPILER_HLSLCC"), 1);
-	AdditionalDefines.SetDefine(TEXT("COMPILER_VECTORVM"), 1);
-	AdditionalDefines.SetDefine(TEXT("VECTORVM_PROFILE"), 1);
-	AdditionalDefines.SetDefine(TEXT("FORCE_FLOATS"), (uint32)1);
+	// disable deprecation warnings due to default-constructed FShaderCompilerDefinitions in function args.
+	// can be removed once FShaderCompilerDefinitions moves to Internal
+	PRAGMA_DISABLE_DEPRECATION_WARNINGS 	
+	// Don't include shader definitions since it creates shader compilation errors.
+	return PreprocessShader(Output, Input, Environment, {}, EDumpShaderDefines::DontIncludeDefines);
 	PRAGMA_ENABLE_DEPRECATION_WARNINGS
-	if (Input.bSkipPreprocessedCache)
-	{
-		if (!FFileHelper::LoadFileToString(Output.EditSource(), *Input.VirtualSourceFilePath))
-		{
-			return false;
-		}
-
-		// Remove const as we are on debug-only mode
-		CrossCompiler::CreateEnvironmentFromResourceTable(Output.GetSource(), (FShaderCompilerEnvironment&)Input.Environment);
-	}
-	else
-	{
-		// Don't include shader definitions since it creates shader compilation errors.
-		if (!PreprocessShader(Output, Input, Environment, AdditionalDefines, EDumpShaderDefines::DontIncludeDefines))
-		{
-			// The preprocessing stage will add any relevant errors.
-			return false;
-		}
-	}
-
-	return true;
 }
 
 bool CompileVectorVMShader(
diff --git a/Engine/Source/Developer/ShaderPreprocessor/Private/ShaderPreprocessor.cpp b/Engine/Source/Developer/ShaderPreprocessor/Private/ShaderPreprocessor.cpp
--- a/Engine/Source/Developer/ShaderPreprocessor/Private/ShaderPreprocessor.cpp
+++ b/Engine/Source/Developer/ShaderPreprocessor/Private/ShaderPreprocessor.cpp
@@ -6,8 +6,8 @@
 #include "Misc/ScopeLock.h"
 #include "Modules/ModuleManager.h"
 #include "PreprocessorPrivate.h"
+#include "ShaderCompilerCommon.h"
 #include "ShaderCompilerDefinitions.h"
-
 #include "stb_preprocess/preprocessor.h"
 #include "stb_preprocess/stb_alloc.h"
 #include "stb_preprocess/stb_ds.h"
@@ -884,22 +884,43 @@ bool PreprocessShader(
 )
 {
 	TRACE_CPUPROFILER_EVENT_SCOPE(PreprocessShader);
+	Output.EditSource().Empty();
 
-	// Skip the cache system and directly load the file path (used for debugging)
-	if (Input.bSkipPreprocessedCache)
+	if (EnumHasAnyFlags(Input.DebugInfoFlags, EShaderDebugInfoFlags::CompileFromDebugUSF))
 	{
-		return FFileHelper::LoadFileToString(Output.EditSource(), *Input.VirtualSourceFilePath);
+		// the "VirtualSourceFilePath" given is actually an absolute path to a dumped debug USF file; load it directly.
+		// this occurs when running SCW in "direct compile" mode; this file will already be preprocessed.
+		bool bSuccess = FFileHelper::LoadFileToString(Output.EditSource(), *Input.VirtualSourceFilePath);
+
+		if (bSuccess)
+		{
+			// const_cast for compile environment; need to populate a subset of environment parameters from parsing comments in the preprocessed code
+			UE::ShaderCompilerCommon::SerializeEnvironmentFromBase64(const_cast<FShaderCompilerEnvironment&>(Input.Environment), Output.GetSource());
+		}
+
+		return bSuccess;
 	}
 
 	check(CheckVirtualShaderFilePath(Input.VirtualSourceFilePath));
 
-	Output.EditSource().Empty();
-
 	// List the defines used for compilation in the preprocessed shaders, especially to know which permutation vector this shader is.
 	if (DefinesPolicy == EDumpShaderDefines::AlwaysIncludeDefines || (DefinesPolicy == EDumpShaderDefines::DontCare && Input.DumpDebugInfoPath.Len() > 0))
 	{
 		FShaderPreprocessorUtilities::DumpShaderDefinesAsCommentedCode(Environment, &Output.EditSource());
 	}
 
-	return InnerPreprocessShaderStb(Output, Input, Environment, AdditionalDefines);
+	bool bSuccess = InnerPreprocessShaderStb(Output, Input, Environment, AdditionalDefines);
+
+	if (bSuccess)
+	{
+		CleanupUniformBufferCode(Environment, Output.EditSource());
+
+		if (Input.Environment.CompilerFlags.Contains(CFLAG_RemoveDeadCode))
+		{
+			const TArray<FStringView> RequiredSymbols(MakeArrayView(Input.RequiredSymbols));
+			UE::ShaderCompilerCommon::RemoveDeadCode(Output.EditSource(), Input.EntryPointName, RequiredSymbols, Output.EditErrors());
+		}
+	}
+
+	return bSuccess;
 }
diff --git a/Engine/Source/Developer/ShaderPreprocessor/Public/ShaderPreprocessor.h b/Engine/Source/Developer/ShaderPreprocessor/Public/ShaderPreprocessor.h
--- a/Engine/Source/Developer/ShaderPreprocessor/Public/ShaderPreprocessor.h
+++ b/Engine/Source/Developer/ShaderPreprocessor/Public/ShaderPreprocessor.h
@@ -60,6 +60,6 @@ extern SHADERPREPROCESSOR_API bool PreprocessShader(
 	const FShaderCompilerInput& Input,
 	const FShaderCompilerEnvironment& MergedEnvironment,
 	PRAGMA_DISABLE_DEPRECATION_WARNINGS		// FShaderCompilerDefinitions will be made internal in the future, marked deprecated until then
-	const FShaderCompilerDefinitions& AdditionalDefines,
+	const FShaderCompilerDefinitions& AdditionalDefines = FShaderCompilerDefinitions(),
 	PRAGMA_ENABLE_DEPRECATION_WARNINGS
 	EDumpShaderDefines DefinesPolicy = EDumpShaderDefines::DontCare);
diff --git a/Engine/Source/Developer/ShaderPreprocessor/ShaderPreprocessor.Build.cs b/Engine/Source/Developer/ShaderPreprocessor/ShaderPreprocessor.Build.cs
--- a/Engine/Source/Developer/ShaderPreprocessor/ShaderPreprocessor.Build.cs
+++ b/Engine/Source/Developer/ShaderPreprocessor/ShaderPreprocessor.Build.cs
@@ -10,6 +10,8 @@ public class ShaderPreprocessor : ModuleRules
 			new string[] {
 				"Core",
 				"RenderCore",
+				"ShaderCompilerCommon",
+				"TargetPlatform",
 			}
 			);
 
diff --git a/Engine/Source/Developer/VulkanShaderFormat/Private/VulkanShaderCompiler.cpp b/Engine/Source/Developer/VulkanShaderFormat/Private/VulkanShaderCompiler.cpp
--- a/Engine/Source/Developer/VulkanShaderFormat/Private/VulkanShaderCompiler.cpp
+++ b/Engine/Source/Developer/VulkanShaderFormat/Private/VulkanShaderCompiler.cpp
@@ -8,7 +8,6 @@
 #include "ShaderCompilerCommon.h"
 #include "ShaderCompilerDefinitions.h"
 #include "ShaderParameterParser.h"
-#include "ShaderPreprocessor.h"
 #include "ShaderPreprocessTypes.h"
 #include "SpirvReflectCommon.h"
 #include "VulkanCommon.h"
@@ -240,8 +239,6 @@ static FString GetBindlessUBNameFromHeap(const FString& HeapName)
 	return HeapName.Mid(NameStart, HeapName.Len() - NameStart - kBindlessHeapSuffix.Len());
 }
 
-// Name of the structure in raytracing shader records in VulkanCommon.usf
-static const FString kHitGroupSystemRootConstantsSymbolName = TEXT("HitGroupSystemRootConstants");
 
 // A collection of states and data that is locked in at the top level call and doesn't change throughout the compilation process
 struct FVulkanShaderCompilerInternalState
@@ -321,38 +318,6 @@ struct FVulkanShaderCompilerInternalState
 			return Input.EntryPointName;
 		}
 	}
-	inline TArray<FStringView> GetRequiredSymbols() const
-	{
-		TArray<FStringView> RequiredSymbols;
-
-		if (bIsRayHitGroupShader)
-		{
-			check(!ClosestHitEntry.IsEmpty());
-
-			RequiredSymbols.Add(ClosestHitEntry);
-
-			if (!AnyHitEntry.IsEmpty())
-			{
-				RequiredSymbols.Add(AnyHitEntry);
-			}
-
-			if (!IntersectionEntry.IsEmpty())
-			{
-				RequiredSymbols.Add(IntersectionEntry);
-			}
-		}
-		else
-		{
-			RequiredSymbols.Add(Input.EntryPointName);
-		}
-
-		if (IsRayTracingShader())
-		{
-			RequiredSymbols.Add(kHitGroupSystemRootConstantsSymbolName);
-		}
-
-		return RequiredSymbols;
-	}
 	inline bool IsRayTracingShader() const
 	{
 		return Input.IsRayTracingShader();
@@ -2099,124 +2064,69 @@ static void RemoveUnusedBindlessHeaps(FString& PreprocessedShaderSource, const T
 	}
 }
 
-
-
-bool PreprocessVulkanShader(const FShaderCompilerInput& Input, const FShaderCompilerEnvironment& Environment, FShaderPreprocessOutput& PreprocessOutput)
+void ModifyVulkanCompilerInput(FShaderCompilerInput& Input)
 {
-	const FVulkanShaderCompilerInternalState InternalState(Input, nullptr);
-
-	PRAGMA_DISABLE_DEPRECATION_WARNINGS		// FShaderCompilerDefinitions will be made internal in the future, marked deprecated until then
-
-	FShaderCompilerDefinitions AdditionalDefines;
-	AdditionalDefines.SetDefine(TEXT("COMPILER_HLSLCC"), 1);
-	AdditionalDefines.SetDefine(TEXT("COMPILER_VULKAN"), 1);
+	FVulkanShaderCompilerInternalState InternalState(Input, nullptr);
+	Input.Environment.SetDefine(TEXT("COMPILER_HLSLCC"), 1);
+	Input.Environment.SetDefine(TEXT("COMPILER_VULKAN"), 1);
 	if (InternalState.IsMobileES31())
 	{
-		AdditionalDefines.SetDefine(TEXT("ES3_1_PROFILE"), 1);
-		AdditionalDefines.SetDefine(TEXT("VULKAN_PROFILE"), 1);
+		Input.Environment.SetDefine(TEXT("ES3_1_PROFILE"), 1);
+		Input.Environment.SetDefine(TEXT("VULKAN_PROFILE"), 1);
 	}
 	else if (InternalState.IsSM6())
 	{
-		AdditionalDefines.SetDefine(TEXT("VULKAN_PROFILE_SM6"), 1);
+		Input.Environment.SetDefine(TEXT("VULKAN_PROFILE_SM6"), 1);
 	}
 	else if (InternalState.IsSM5())
 	{
-		AdditionalDefines.SetDefine(TEXT("VULKAN_PROFILE_SM5"), 1);
+		Input.Environment.SetDefine(TEXT("VULKAN_PROFILE_SM5"), 1);
 	}
-	AdditionalDefines.SetDefine(TEXT("row_major"), TEXT(""));
+	Input.Environment.SetDefine(TEXT("row_major"), TEXT(""));
 
-	AdditionalDefines.SetDefine(TEXT("COMPILER_SUPPORTS_ATTRIBUTES"), (uint32)1);
-	AdditionalDefines.SetDefine(TEXT("COMPILER_SUPPORTS_DUAL_SOURCE_BLENDING_SLOT_DECORATION"), (uint32)1);
-	AdditionalDefines.SetDefine(TEXT("PLATFORM_SUPPORTS_ROV"), 0); // Disabled until DXC->SPRIV ROV support is implemented
+	Input.Environment.SetDefine(TEXT("COMPILER_SUPPORTS_ATTRIBUTES"), (uint32)1);
+	Input.Environment.SetDefine(TEXT("COMPILER_SUPPORTS_DUAL_SOURCE_BLENDING_SLOT_DECORATION"), (uint32)1);
+	Input.Environment.SetDefine(TEXT("PLATFORM_SUPPORTS_ROV"), 0); // Disabled until DXC->SPRIV ROV support is implemented
 
-	if (Input.Environment.FullPrecisionInPS)
+	if (Input.Environment.FullPrecisionInPS || (IsValidRef(Input.SharedEnvironment) && Input.SharedEnvironment->FullPrecisionInPS))
 	{
-		AdditionalDefines.SetDefine(TEXT("FORCE_FLOATS"), (uint32)1);
+		Input.Environment.SetDefine(TEXT("FORCE_FLOATS"), (uint32)1);
 	}
 
 	if (Input.Environment.CompilerFlags.Contains(CFLAG_InlineRayTracing))
 	{
-		AdditionalDefines.SetDefine(TEXT("PLATFORM_SUPPORTS_INLINE_RAY_TRACING"), 1);
+		Input.Environment.SetDefine(TEXT("PLATFORM_SUPPORTS_INLINE_RAY_TRACING"), 1);
 	}
 
 	if (Input.Environment.CompilerFlags.Contains(CFLAG_AllowRealTypes))
 	{
-		AdditionalDefines.SetDefine(TEXT("PLATFORM_SUPPORTS_REAL_TYPES"), 1);
+		Input.Environment.SetDefine(TEXT("PLATFORM_SUPPORTS_REAL_TYPES"), 1);
 	}
 
 	// We have ETargetEnvironment::Vulkan_1_1 by default as a min spec now
 	{
-		AdditionalDefines.SetDefine(TEXT("PLATFORM_SUPPORTS_SM6_0_WAVE_OPERATIONS"), 1);
-		AdditionalDefines.SetDefine(TEXT("VULKAN_SUPPORTS_SUBGROUP_SIZE_CONTROL"), 1);
+		Input.Environment.SetDefine(TEXT("PLATFORM_SUPPORTS_SM6_0_WAVE_OPERATIONS"), 1);
+		Input.Environment.SetDefine(TEXT("VULKAN_SUPPORTS_SUBGROUP_SIZE_CONTROL"), 1);
 	}
 
-	AdditionalDefines.SetDefine(TEXT("VULKAN_BINDLESS_SRV_ARRAY_PREFIX"), FShaderParameterParser::kBindlessSRVArrayPrefix);
-	AdditionalDefines.SetDefine(TEXT("VULKAN_BINDLESS_UAV_ARRAY_PREFIX"), FShaderParameterParser::kBindlessUAVArrayPrefix);
-	AdditionalDefines.SetDefine(TEXT("VULKAN_BINDLESS_SAMPLER_ARRAY_PREFIX"), FShaderParameterParser::kBindlessSamplerArrayPrefix);
-	AdditionalDefines.SetDefine(TEXT("VULKAN_MAX_BINDLESS_UNIFORM_BUFFERS_PER_STAGE"), VulkanBindless::MaxUniformBuffersPerStage);
+	Input.Environment.SetDefine(TEXT("VULKAN_BINDLESS_SRV_ARRAY_PREFIX"), FShaderParameterParser::kBindlessSRVArrayPrefix);
+	Input.Environment.SetDefine(TEXT("VULKAN_BINDLESS_UAV_ARRAY_PREFIX"), FShaderParameterParser::kBindlessUAVArrayPrefix);
+	Input.Environment.SetDefine(TEXT("VULKAN_BINDLESS_SAMPLER_ARRAY_PREFIX"), FShaderParameterParser::kBindlessSamplerArrayPrefix);
+	Input.Environment.SetDefine(TEXT("VULKAN_MAX_BINDLESS_UNIFORM_BUFFERS_PER_STAGE"), VulkanBindless::MaxUniformBuffersPerStage);
 
 	if (IsAndroidShaderFormat(Input.ShaderFormat))
 	{
 		// On most Android devices uint64_t is unsupported so we emulate as 2 uint32_t's 
-		AdditionalDefines.SetDefine(TEXT("EMULATE_VKDEVICEADRESS"), 1);
-	}
-
-	PRAGMA_ENABLE_DEPRECATION_WARNINGS
-
-	// Preprocess the shader.
-	FString& PreprocessedShaderSource = PreprocessOutput.EditSource();
-	const bool bDirectCompile = FParse::Param(FCommandLine::Get(), TEXT("directcompile"));
-	if (bDirectCompile)
-	{
-		if (!FFileHelper::LoadFileToString(PreprocessedShaderSource, *Input.VirtualSourceFilePath))
-		{
-			return false;
-		}
-
-		// Remove const as we are on debug-only mode
-		CrossCompiler::CreateEnvironmentFromResourceTable(PreprocessedShaderSource, (FShaderCompilerEnvironment&)Input.Environment);
-	}
-	else
-	{
-		if (!PreprocessShader(PreprocessOutput, Input, Environment, AdditionalDefines))
-		{
-			// The preprocessing stage will add any relevant errors.
-			return false;
-		}
-	}
-
-	CleanupUniformBufferCode(Input.Environment, PreprocessedShaderSource);
-
-	// Run the shader minifier
-	if (InternalState.IsRayTracingShader() || InternalState.bUseBindlessUniformBuffer)
-	{
-		// Always needed in ray tracing shaders to ensure bindless UB count is kept low
-		UE::ShaderCompilerCommon::RemoveDeadCode(PreprocessedShaderSource, InternalState.GetRequiredSymbols(), PreprocessOutput.EditErrors());
+		Input.Environment.SetDefine(TEXT("EMULATE_VKDEVICEADRESS"), 1);
 	}
-	#if UE_VULKAN_SHADER_COMPILER_ALLOW_DEAD_CODE_REMOVAL
-	else if (Input.Environment.CompilerFlags.Contains(CFLAG_RemoveDeadCode))
-	{
-		UE::ShaderCompilerCommon::RemoveDeadCode(PreprocessedShaderSource, InternalState.GetEntryPointName(), PreprocessOutput.EditErrors());
-	}
-	#endif // UE_VULKAN_SHADER_COMPILER_ALLOW_DEAD_CODE_REMOVAL
 
-	// Clean up the code a bit, it's unreadable otherwise with all the unused heaps left around
-	if (InternalState.bSupportsBindless)
+	if (Input.IsRayTracingShader())
 	{
-		RemoveUnusedBindlessHeaps(PreprocessedShaderSource, TEXT("SAMPLER"));
-		RemoveUnusedBindlessHeaps(PreprocessedShaderSource, TEXT("RESOURCE"));
-		UE::ShaderCompilerCommon::RemoveDeadCode(PreprocessedShaderSource, InternalState.GetRequiredSymbols(), PreprocessOutput.EditErrors());
-
-		if (InternalState.bDebugDump)
-		{
-			DumpDebugShaderText(Input, PreprocessedShaderSource, TEXT("bindless.final.hlsl"));
-		}
+		// Name of the structure in raytracing shader records in VulkanCommon.usf
+		Input.RequiredSymbols.Add(TEXT("HitGroupSystemRootConstants"));
 	}
-
-	return true;
 }
 
-
 // :todo-jn: TEMPORARY EXPERIMENT - will eventually move into preprocessing step
 static TArray<FString> ConvertUBToBindless(FString& PreprocessedShaderSource)
 {
@@ -2522,6 +2432,22 @@ void CompileVulkanShader(const FShaderCompilerInput& Input, const FString& InPre
 
 	FVulkanShaderCompilerInternalState InternalState(Input, &ShaderParameterParser);
 
+	//TODO: this additional step causes problems for the error remapping that occurs when the preprocessed job cache is enabled
+	// (the additional deadstripping step causes further changes to line numbers, removed blocks, etc).
+	//if (InternalState.bSupportsBindless)
+	//{
+	//	// Clean up the code a bit, it's unreadable otherwise with all the unused heaps left around
+	//	// Re-run the dead stripper after removing these unused heaps
+	//	RemoveUnusedBindlessHeaps(PreprocessedSource, TEXT("SAMPLER"));
+	//	RemoveUnusedBindlessHeaps(PreprocessedSource, TEXT("RESOURCE"));
+	//	UE::ShaderCompilerCommon::RemoveDeadCode(PreprocessedSource, Input.EntryPointName, Input.RequiredSymbols, Output.Errors);
+
+	//	if (InternalState.bDebugDump)
+	//	{
+	//		DumpDebugShaderText(Input, PreprocessedSource, TEXT("bindless.final.hlsl"));
+	//	}
+	//}
+
 	const EHlslShaderFrequency HlslFrequency = InternalState.GetHlslShaderFrequency();
 	if (HlslFrequency == HSF_InvalidFrequency)
 	{
@@ -2586,8 +2512,7 @@ void CompileVulkanShader(const FShaderCompilerInput& Input, const FString& InPre
 
 	ShaderParameterParser.ValidateShaderParameterTypes(Input, InternalState.IsMobileES31(), Output);
 	
-	const bool bDirectCompile = FParse::Param(FCommandLine::Get(), TEXT("directcompile"));
-	if (bDirectCompile)
+	if (EnumHasAnyFlags(Input.DebugInfoFlags, EShaderDebugInfoFlags::CompileFromDebugUSF))
 	{
 		for (const auto& Error : Output.Errors)
 		{
diff --git a/Engine/Source/Developer/VulkanShaderFormat/Private/VulkanShaderFormat.cpp b/Engine/Source/Developer/VulkanShaderFormat/Private/VulkanShaderFormat.cpp
--- a/Engine/Source/Developer/VulkanShaderFormat/Private/VulkanShaderFormat.cpp
+++ b/Engine/Source/Developer/VulkanShaderFormat/Private/VulkanShaderFormat.cpp
@@ -8,16 +8,16 @@
 #include "Interfaces/IShaderFormatModule.h"
 #include "hlslcc.h"
 #include "ShaderCore.h"
+#include "ShaderCompilerCommon.h"
 #include "ShaderCompilerCore.h"
+#include "ShaderParameterParser.h"
+#include "ShaderPreprocessor.h"
 #include "ShaderPreprocessTypes.h"
 #include "DXCWrapper.h"
 #include "ShaderConductorContext.h"
 #include "RHIShaderFormatDefinitions.inl"
 
-extern bool PreprocessVulkanShader(
-	const FShaderCompilerInput& Input,
-	const FShaderCompilerEnvironment& Environment,
-	FShaderPreprocessOutput& PreprocessOutput);
+extern void ModifyVulkanCompilerInput(FShaderCompilerInput& Input);
 
 extern void CompileVulkanShader(
 	const FShaderCompilerInput& Input,
@@ -30,9 +30,9 @@ extern void OutputVulkanDebugData(
 	const FShaderPreprocessOutput& PreprocessOutput, 
 	const FShaderCompilerOutput& Output);
 
-static const FGuid UE_SHADER_VULKAN_ES3_1_VER = FGuid("6D333063-D2F7-4AA5-A79D-BA73F32C898E");
-static const FGuid UE_SHADER_VULKAN_SM5_VER = FGuid("6EC81E81-BDE8-4F09-8A70-AD09A817B32A");
-static const FGuid UE_SHADER_VULKAN_SM6_VER = FGuid("C732FBB7-4CAD-4249-A760-912777047233");
+static const FGuid UE_SHADER_VULKAN_ES3_1_VER = FGuid("B84F72C8-3ECD-411E-993C-D7C7CEE26F28");
+static const FGuid UE_SHADER_VULKAN_SM5_VER = FGuid("0715D8EE-9907-4A25-93AD-A3902C8E069A");
+static const FGuid UE_SHADER_VULKAN_SM6_VER = FGuid("C5161730-83C6-40AF-A990-78CD4C1581DB");
 
 class FShaderFormatVulkan : public IShaderFormat
 {
@@ -77,10 +77,6 @@ public:
 		Version = HashCombine(Version, 0xFC0848E2);
 	#endif
 
-	#if UE_VULKAN_SHADER_COMPILER_ALLOW_DEAD_CODE_REMOVAL
-		Version = HashCombine(Version, 0x75E2FE85);
-	#endif // UE_VULKAN_SHADER_COMPILER_ALLOW_DEAD_CODE_REMOVAL
-
 		return Version;
 	}
 	virtual void GetSupportedFormats(TArray<FName>& OutFormats) const
@@ -92,9 +88,14 @@ public:
 		OutFormats.Add(NAME_VULKAN_SM6);
 	}
 
+	virtual void ModifyShaderCompilerInput(FShaderCompilerInput& Input) const override
+	{
+		ModifyVulkanCompilerInput(Input);
+	}
+
 	virtual bool PreprocessShader(const FShaderCompilerInput& Input, const FShaderCompilerEnvironment& Environment, FShaderPreprocessOutput& PreprocessOutput) const
 	{
-		return PreprocessVulkanShader(Input, Environment, PreprocessOutput);
+		return ::PreprocessShader(PreprocessOutput, Input, Environment);
 	}
 
 	virtual void CompilePreprocessedShader(const FShaderCompilerInput& Input, const FShaderPreprocessOutput& PreprocessOutput, FShaderCompilerOutput& Output,const FString& WorkingDirectory) const override
diff --git a/Engine/Source/Developer/VulkanShaderFormat/Private/VulkanShaderFormat.h b/Engine/Source/Developer/VulkanShaderFormat/Private/VulkanShaderFormat.h
--- a/Engine/Source/Developer/VulkanShaderFormat/Private/VulkanShaderFormat.h
+++ b/Engine/Source/Developer/VulkanShaderFormat/Private/VulkanShaderFormat.h
@@ -7,12 +7,6 @@
 #include "hlslcc.h"
 
 
-// Controls whether r.Shaders.RemoveDeadCode should be honored
-#ifndef UE_VULKAN_SHADER_COMPILER_ALLOW_DEAD_CODE_REMOVAL
-#define UE_VULKAN_SHADER_COMPILER_ALLOW_DEAD_CODE_REMOVAL 1
-#endif // UE_VULKAN_SHADER_COMPILER_ALLOW_DEAD_CODE_REMOVAL
-
-
 struct FVulkanSpirv : FSpirv
 {
 	struct FEntry
diff --git a/Engine/Source/Developer/Windows/ShaderFormatD3D/Private/D3DShaderCompiler.cpp b/Engine/Source/Developer/Windows/ShaderFormatD3D/Private/D3DShaderCompiler.cpp
--- a/Engine/Source/Developer/Windows/ShaderFormatD3D/Private/D3DShaderCompiler.cpp
+++ b/Engine/Source/Developer/Windows/ShaderFormatD3D/Private/D3DShaderCompiler.cpp
@@ -7,7 +7,6 @@
 #include "RayTracingDefinitions.h"
 #include "Serialization/MemoryWriter.h"
 #include "ShaderFormatD3D.h"
-#include "ShaderPreprocessor.h"
 #include "ShaderCompilerCommon.h"
 #include "ShaderCompilerDefinitions.h"
 #include "ShaderMinifier.h"
@@ -1260,49 +1259,6 @@ bool CompileAndProcessD3DShaderFXC(
 	return bSuccess;
 }
 
-bool PreprocessD3DShader(
-	const FShaderCompilerInput& Input,
-	const FShaderCompilerEnvironment& Environment,
-	FShaderPreprocessOutput& Output)
-{
-	PRAGMA_DISABLE_DEPRECATION_WARNINGS		// FShaderCompilerDefinitions will be made internal in the future, marked deprecated until then
-	FShaderCompilerDefinitions AdditionalDefines;
-	PRAGMA_ENABLE_DEPRECATION_WARNINGS
-
-	FString& PreprocessedSource = Output.EditSource();
-	
-	if (Input.bSkipPreprocessedCache)
-	{
-		if (!FFileHelper::LoadFileToString(PreprocessedSource, *Input.VirtualSourceFilePath))
-		{
-			return false;
-		}
-
-		// Remove const as we are on debug-only mode
-		CrossCompiler::CreateEnvironmentFromResourceTable(PreprocessedSource, const_cast<FShaderCompilerEnvironment&>(Environment));
-	}
-	else
-	{
-		if (!PreprocessShader(Output, Input, Environment, AdditionalDefines))
-		{
-			// The preprocessing stage will add any relevant errors.
-			return false;
-		}
-	}
-
-	CleanupUniformBufferCode(Input.Environment, PreprocessedSource);
-
-	// Run the shader minifier
-	#if UE_D3D_SHADER_COMPILER_ALLOW_DEAD_CODE_REMOVAL
-	if (Environment.CompilerFlags.Contains(CFLAG_RemoveDeadCode))
-	{
-		UE::ShaderCompilerCommon::RemoveDeadCode(PreprocessedSource, Input.EntryPointName, Output.EditErrors());
-	}
-	#endif // UE_D3D_SHADER_COMPILER_ALLOW_DEAD_CODE_REMOVAL
-
-	return true;
-}
-
 struct FD3DShaderParameterParserPlatformConfiguration : public FShaderParameterParser::FPlatformConfiguration
 {
 	FD3DShaderParameterParserPlatformConfiguration()
@@ -1426,8 +1382,7 @@ void CompileD3DShader(const FShaderCompilerInput& Input, const FString& InPrepro
 
 	ShaderParameterParser.ValidateShaderParameterTypes(Input, Output);
 
-	const bool bDirectCompile = FParse::Param(FCommandLine::Get(), TEXT("directcompile"));
-	if (bDirectCompile)
+	if (EnumHasAnyFlags(Input.DebugInfoFlags, EShaderDebugInfoFlags::CompileFromDebugUSF))
 	{
 		for (const FShaderCompilerError& Error : Output.Errors)
 		{
diff --git a/Engine/Source/Developer/Windows/ShaderFormatD3D/Private/ShaderFormatD3D.cpp b/Engine/Source/Developer/Windows/ShaderFormatD3D/Private/ShaderFormatD3D.cpp
--- a/Engine/Source/Developer/Windows/ShaderFormatD3D/Private/ShaderFormatD3D.cpp
+++ b/Engine/Source/Developer/Windows/ShaderFormatD3D/Private/ShaderFormatD3D.cpp
@@ -2,6 +2,7 @@
 
 #include "ShaderFormatD3D.h"
 #include "ShaderCompilerCommon.h"
+#include "ShaderPreprocessor.h"
 #include "ShaderPreprocessTypes.h"
 #include "ShaderSymbolExport.h"
 #include "Modules/ModuleInterface.h"
@@ -14,10 +15,10 @@ static FName NAME_PCD3D_SM6(TEXT("PCD3D_SM6"));
 static FName NAME_PCD3D_SM5(TEXT("PCD3D_SM5"));
 static FName NAME_PCD3D_ES3_1(TEXT("PCD3D_ES31"));
 
-static const FGuid UE_SHADER_PCD3D_SHARED_VER = FGuid("dd4e6e76-4b48-4097-9ece-0f21118b7177");
-static const FGuid UE_SHADER_PCD3D_SM6_VER    = FGuid("EEDED915-D5D7-4F45-B682-2657DCA0C297");
-static const FGuid UE_SHADER_PCD3D_SM5_VER    = FGuid("c2f60605-b462-4449-9f76-94bd2b1cd78e");
-static const FGuid UE_SHADER_PCD3D_ES3_1_VER  = FGuid("75466d2b-e169-40d8-bac5-1e2f9d43e0bb");
+static const FGuid UE_SHADER_PCD3D_SHARED_VER = FGuid("232A2A59-A6D0-4CDB-A374-F3DB028E413E");
+static const FGuid UE_SHADER_PCD3D_SM6_VER    = FGuid("7BDBA1A9-FAEF-42A1-BCF7-1B921FDEAD9E");
+static const FGuid UE_SHADER_PCD3D_SM5_VER    = FGuid("5B377D13-C70F-40C5-80C5-C9B228783469");
+static const FGuid UE_SHADER_PCD3D_ES3_1_VER  = FGuid("952939D9-1156-4347-97E9-9FFEA1A9FE14");
 
 class FShaderFormatD3D : public UE::ShaderCompilerCommon::FBaseShaderFormat 
 {
@@ -40,10 +41,6 @@ public:
 		const uint32 BaseHash = GetTypeHash(UE_SHADER_PCD3D_SHARED_VER);
 		uint32 VersionHash = GetTypeHash(InVersion);
 
-	#if UE_D3D_SHADER_COMPILER_ALLOW_DEAD_CODE_REMOVAL
-		VersionHash = HashCombine(VersionHash, 0x75E2FE85);
-	#endif // UE_D3D_SHADER_COMPILER_ALLOW_DEAD_CODE_REMOVAL
-
 		return HashCombine(BaseHash, VersionHash);
 	}
 
@@ -176,7 +173,7 @@ public:
 
 	virtual bool PreprocessShader(const FShaderCompilerInput& Input, const FShaderCompilerEnvironment& MergedEnvironment, FShaderPreprocessOutput& PreprocessOutput) const
 	{
-		return PreprocessD3DShader(Input, MergedEnvironment, PreprocessOutput);
+		return ::PreprocessShader(PreprocessOutput, Input, MergedEnvironment);
 	}
 
 	virtual void CompilePreprocessedShader(
diff --git a/Engine/Source/Developer/Windows/ShaderFormatD3D/Private/ShaderFormatD3D.h b/Engine/Source/Developer/Windows/ShaderFormatD3D/Private/ShaderFormatD3D.h
--- a/Engine/Source/Developer/Windows/ShaderFormatD3D/Private/ShaderFormatD3D.h
+++ b/Engine/Source/Developer/Windows/ShaderFormatD3D/Private/ShaderFormatD3D.h
@@ -5,11 +5,6 @@
 #include "Windows/WindowsHWrapper.h"
 #include "ShaderCompilerCommon.h"
 
-// Controls whether r.Shaders.RemoveDeadCode should be honored
-#ifndef UE_D3D_SHADER_COMPILER_ALLOW_DEAD_CODE_REMOVAL
-#define UE_D3D_SHADER_COMPILER_ALLOW_DEAD_CODE_REMOVAL 1
-#endif // UE_D3D_SHADER_COMPILER_ALLOW_DEAD_CODE_REMOVAL
-
 struct FShaderTarget;
 
 enum class ED3DShaderModel
diff --git a/Engine/Source/Programs/ShaderCompileWorker/Private/ShaderCompileWorker.cpp b/Engine/Source/Programs/ShaderCompileWorker/Private/ShaderCompileWorker.cpp
--- a/Engine/Source/Programs/ShaderCompileWorker/Private/ShaderCompileWorker.cpp
+++ b/Engine/Source/Programs/ShaderCompileWorker/Private/ShaderCompileWorker.cpp
@@ -895,10 +895,6 @@ static void DirectCompile(const TArray<const class IShaderFormat*>& ShaderFormat
 			{
 				ShaderPlatformName = FName(*Token.RightChop(19));
 			}
-			else if (Token.StartsWith(TEXT("cflags=")))
-			{
-				CFlags = FCString::Atoi64(*Token.RightChop(7));
-			}
 			else if (Token.StartsWith(TEXT("DebugInfoPath=")))
 			{
 				DumpDebugInfoPath = Token.RightChop(14);
@@ -980,12 +976,11 @@ static void DirectCompile(const TArray<const class IShaderFormat*>& ShaderFormat
 	Job.Input.VirtualSourceFilePath = InputFile;
 	Job.Input.Target.Platform =  ShaderFormatNameToShaderPlatform(FormatName);
 	Job.Input.Target.Frequency = Frequency;
-	Job.Input.bSkipPreprocessedCache = true;
 	Job.Input.DumpDebugInfoPath = DumpDebugInfoPath;
-	Job.Input.Environment.CompilerFlags = FShaderCompilerFlags(CFlags);
 	Job.Input.bCompilingForShaderPipeline = bPipeline;
 	Job.Input.bIncludeUsedOutputs = bIncludeUsedOutputs;
 	Job.Input.UsedOutputs = UsedOutputs;
+	Job.Input.DebugInfoFlags = EShaderDebugInfoFlags::CompileFromDebugUSF;
 
 	FShaderCompilerOutput Output;
 	CompileShader(GetShaderFormats(), Job, Dir, &GNumProcessedJobs);
diff --git a/Engine/Source/Runtime/RenderCore/Private/ShaderCore.cpp b/Engine/Source/Runtime/RenderCore/Private/ShaderCore.cpp
--- a/Engine/Source/Runtime/RenderCore/Private/ShaderCore.cpp
+++ b/Engine/Source/Runtime/RenderCore/Private/ShaderCore.cpp
@@ -3537,7 +3537,6 @@ FArchive& operator<<(FArchive& Ar, FShaderCompilerInput& Input)
 	Ar << Input.VirtualSourceFilePath;
 	Ar << Input.EntryPointName;
 	Ar << Input.ShaderName;
-	Ar << Input.bSkipPreprocessedCache;
 	Ar << Input.bCompilingForShaderPipeline;
 	Ar << Input.bIncludeUsedOutputs;
 	Ar << Input.bCachePreprocessed;
diff --git a/Engine/Source/Runtime/RenderCore/Public/ShaderCompilerCore.h b/Engine/Source/Runtime/RenderCore/Public/ShaderCompilerCore.h
--- a/Engine/Source/Runtime/RenderCore/Public/ShaderCompilerCore.h
+++ b/Engine/Source/Runtime/RenderCore/Public/ShaderCompilerCore.h
@@ -31,7 +31,7 @@ class FShaderPipelineCompileJob;
 typedef TSharedPtr<TArray<ANSICHAR>, ESPMode::ThreadSafe> FShaderSharedAnsiStringPtr;
 
 // this is for the protocol, not the data, bump if FShaderCompilerInput or ProcessInputFromArchive changes.
-inline const int32 ShaderCompileWorkerInputVersion = 24;
+inline const int32 ShaderCompileWorkerInputVersion = 25;
 // this is for the protocol, not the data, bump if FShaderCompilerOutput or WriteToOutputArchive changes.
 inline const int32 ShaderCompileWorkerOutputVersion = 19;
 // this is for the protocol, not the data.
@@ -221,6 +221,7 @@ enum class EShaderDebugInfoFlags : uint8
 	Diagnostics = 1 << 2,
 	ShaderCodeBinary = 1 << 3,
 	DetailedSource = 1 << 4,
+	CompileFromDebugUSF = 1 << 5,
 };
 ENUM_CLASS_FLAGS(EShaderDebugInfoFlags)
 
@@ -238,6 +239,7 @@ struct FShaderCompilerInput
 	FString ShaderName;
 
 	// Skips the preprocessor and instead loads the usf file directly
+	UE_DEPRECATED(5.4, "bSkipPreprocessedCache member is deprecated; set EShaderDebugInfoFlags::CompileFromDebugUSF on DebugInfoFlags instead.")
 	bool bSkipPreprocessedCache;
 
 	UE_DEPRECATED(5.3, "Use DebugInfoFlags field (EDebugInfoFlags::DirectCompileCommandLine)")
@@ -255,6 +257,10 @@ struct FShaderCompilerInput
 	// as such this will only ever be set for jobs whose shader format supports independent
 	// preprocessing)
 	bool bCachePreprocessed;
+
+	// Array of symbols that should be maintained when deadstripping. If this is empty, entry
+	// point name alone will be used.
+	TArray<FString> RequiredSymbols;
 	
 	// Shader pipeline information
 	bool bCompilingForShaderPipeline;
diff --git a/Engine/Source/Runtime/RenderCore/Public/ShaderCore.h b/Engine/Source/Runtime/RenderCore/Public/ShaderCore.h
--- a/Engine/Source/Runtime/RenderCore/Public/ShaderCore.h
+++ b/Engine/Source/Runtime/RenderCore/Public/ShaderCore.h
@@ -417,6 +417,12 @@ public:
 		Data = Data | FlagBit;
 	}
 
+	inline void Remove(uint32 InFlag)
+	{
+		const uint64 FlagBit = (uint64)1 << (uint64)InFlag;
+		Data = Data & ~FlagBit;
+	}
+
 	inline bool Contains(uint32 InFlag) const
 	{
 		const uint64 FlagBit = (uint64)1 << (uint64)InFlag;
