Fixing Vulkan compiles on Mac.

diff --git a/Engine/Source/Developer/VulkanShaderFormat/Private/VulkanShaderCompiler.cpp b/Engine/Source/Developer/VulkanShaderFormat/Private/VulkanShaderCompiler.cpp
--- a/Engine/Source/Developer/VulkanShaderFormat/Private/VulkanShaderCompiler.cpp
+++ b/Engine/Source/Developer/VulkanShaderFormat/Private/VulkanShaderCompiler.cpp
@@ -12,33 +12,7 @@
 #include "SpirvReflectCommon.h"
 #include "VulkanCommon.h"
 
-#if PLATFORM_MAC
-// Horrible hack as we need the enum available but the Vulkan headers do not compile on Mac
-typedef enum VkDescriptorType {
-    VK_DESCRIPTOR_TYPE_SAMPLER = 0,
-    VK_DESCRIPTOR_TYPE_COMBINED_IMAGE_SAMPLER = 1,
-    VK_DESCRIPTOR_TYPE_SAMPLED_IMAGE = 2,
-    VK_DESCRIPTOR_TYPE_STORAGE_IMAGE = 3,
-    VK_DESCRIPTOR_TYPE_UNIFORM_TEXEL_BUFFER = 4,
-    VK_DESCRIPTOR_TYPE_STORAGE_TEXEL_BUFFER = 5,
-    VK_DESCRIPTOR_TYPE_UNIFORM_BUFFER = 6,
-    VK_DESCRIPTOR_TYPE_STORAGE_BUFFER = 7,
-    VK_DESCRIPTOR_TYPE_UNIFORM_BUFFER_DYNAMIC = 8,
-    VK_DESCRIPTOR_TYPE_STORAGE_BUFFER_DYNAMIC = 9,
-    VK_DESCRIPTOR_TYPE_INPUT_ATTACHMENT = 10,
-    VK_DESCRIPTOR_TYPE_INLINE_UNIFORM_BLOCK = 1000138000,
-    VK_DESCRIPTOR_TYPE_ACCELERATION_STRUCTURE_KHR = 1000150000,
-    VK_DESCRIPTOR_TYPE_ACCELERATION_STRUCTURE_NV = 1000165000,
-    VK_DESCRIPTOR_TYPE_SAMPLE_WEIGHT_IMAGE_QCOM = 1000440000,
-    VK_DESCRIPTOR_TYPE_BLOCK_MATCH_IMAGE_QCOM = 1000440001,
-    VK_DESCRIPTOR_TYPE_MUTABLE_EXT = 1000351000,
-    VK_DESCRIPTOR_TYPE_INLINE_UNIFORM_BLOCK_EXT = VK_DESCRIPTOR_TYPE_INLINE_UNIFORM_BLOCK,
-    VK_DESCRIPTOR_TYPE_MUTABLE_VALVE = VK_DESCRIPTOR_TYPE_MUTABLE_EXT,
-    VK_DESCRIPTOR_TYPE_MAX_ENUM = 0x7FFFFFFF
-} VkDescriptorType;
-#else
-#include "IVulkanDynamicRHI.h"
-#endif
+#include "VulkanThirdParty.h"
 #include "VulkanBackend.h"
 #include "VulkanShaderResources.h"
 #include "Serialization/MemoryWriter.h"
diff --git a/Engine/Source/Runtime/VulkanRHI/Public/VulkanThirdParty.h b/Engine/Source/Runtime/VulkanRHI/Public/VulkanThirdParty.h
--- a/Engine/Source/Runtime/VulkanRHI/Public/VulkanThirdParty.h
+++ b/Engine/Source/Runtime/VulkanRHI/Public/VulkanThirdParty.h
@@ -8,7 +8,35 @@
 	#include "Microsoft/HideMicrosoftPlatformTypes.h"
 #endif
 
+#if PLATFORM_MAC
+// Horrible hack as we need the enum available but the Vulkan headers do not compile on Mac
+typedef enum VkDescriptorType {
+	VK_DESCRIPTOR_TYPE_SAMPLER = 0,
+	VK_DESCRIPTOR_TYPE_COMBINED_IMAGE_SAMPLER = 1,
+	VK_DESCRIPTOR_TYPE_SAMPLED_IMAGE = 2,
+	VK_DESCRIPTOR_TYPE_STORAGE_IMAGE = 3,
+	VK_DESCRIPTOR_TYPE_UNIFORM_TEXEL_BUFFER = 4,
+	VK_DESCRIPTOR_TYPE_STORAGE_TEXEL_BUFFER = 5,
+	VK_DESCRIPTOR_TYPE_UNIFORM_BUFFER = 6,
+	VK_DESCRIPTOR_TYPE_STORAGE_BUFFER = 7,
+	VK_DESCRIPTOR_TYPE_UNIFORM_BUFFER_DYNAMIC = 8,
+	VK_DESCRIPTOR_TYPE_STORAGE_BUFFER_DYNAMIC = 9,
+	VK_DESCRIPTOR_TYPE_INPUT_ATTACHMENT = 10,
+	VK_DESCRIPTOR_TYPE_INLINE_UNIFORM_BLOCK = 1000138000,
+	VK_DESCRIPTOR_TYPE_ACCELERATION_STRUCTURE_KHR = 1000150000,
+	VK_DESCRIPTOR_TYPE_ACCELERATION_STRUCTURE_NV = 1000165000,
+	VK_DESCRIPTOR_TYPE_SAMPLE_WEIGHT_IMAGE_QCOM = 1000440000,
+	VK_DESCRIPTOR_TYPE_BLOCK_MATCH_IMAGE_QCOM = 1000440001,
+	VK_DESCRIPTOR_TYPE_MUTABLE_EXT = 1000351000,
+	VK_DESCRIPTOR_TYPE_INLINE_UNIFORM_BLOCK_EXT = VK_DESCRIPTOR_TYPE_INLINE_UNIFORM_BLOCK,
+	VK_DESCRIPTOR_TYPE_MUTABLE_VALVE = VK_DESCRIPTOR_TYPE_MUTABLE_EXT,
+	VK_DESCRIPTOR_TYPE_MAX_ENUM = 0x7FFFFFFF
+} VkDescriptorType;
+#else
+
 THIRD_PARTY_INCLUDES_START
 	#define VK_NO_PROTOTYPES
 	#include <vulkan.h>
 THIRD_PARTY_INCLUDES_END
+
+#endif
