[shaders] add ANSICHAR path for FShaderSource and enable by default

diff --git a/Engine/Source/Developer/ShaderCompilerCommon/Private/ShaderCompilerCommon.cpp b/Engine/Source/Developer/ShaderCompilerCommon/Private/ShaderCompilerCommon.cpp
--- a/Engine/Source/Developer/ShaderCompilerCommon/Private/ShaderCompilerCommon.cpp
+++ b/Engine/Source/Developer/ShaderCompilerCommon/Private/ShaderCompilerCommon.cpp
@@ -467,7 +467,20 @@ bool UE::ShaderCompilerCommon::RemoveDeadCode(FShaderSource& InOutPreprocessedSh
 	           |  UE::ShaderMinifier::EMinifyShaderFlags::OutputStats;  // Output a comment detailing how many blocks of each type (functions/structs/etc.) were emitted
 #endif
 
-	UE::ShaderMinifier::FMinifiedShader Minified = UE::ShaderMinifier::Minify(InOutPreprocessedShaderSource, InRequiredSymbols,
+#if SHADER_SOURCE_ANSI
+	TArray<FShaderSource::FStringType> ConvertedRequiredSymbols;
+	TArray<FShaderSource::FViewType> RequiredSymbolViews;
+	for (FStringView InSymbol : InRequiredSymbols)
+	{
+		FShaderSource::FStringType& ConvertedString = ConvertedRequiredSymbols.AddDefaulted_GetRef();
+		ConvertedString.Append(InSymbol);
+		RequiredSymbolViews.Add(FShaderSource::FViewType(ConvertedString));
+	}
+#else
+	TConstArrayView<FStringView> RequiredSymbolViews = InRequiredSymbols;
+#endif
+
+	UE::ShaderMinifier::FMinifiedShader Minified  = UE::ShaderMinifier::Minify(InOutPreprocessedShaderSource, RequiredSymbolViews,
 		  UE::ShaderMinifier::EMinifyShaderFlags::OutputCommentLines // Preserve comments that were left after preprocessing
 		| UE::ShaderMinifier::EMinifyShaderFlags::OutputLines        // Emit #line directives
 		| ExtraFlags);
diff --git a/Engine/Source/Developer/ShaderCompilerCommon/Public/ShaderMinifier.h b/Engine/Source/Developer/ShaderCompilerCommon/Public/ShaderMinifier.h
--- a/Engine/Source/Developer/ShaderCompilerCommon/Public/ShaderMinifier.h
+++ b/Engine/Source/Developer/ShaderCompilerCommon/Public/ShaderMinifier.h
@@ -6,7 +6,7 @@
 
 #include "Containers/UnrealString.h"
 #include "Misc/EnumClassFlags.h"
-#include "ShaderPreprocessTypes.h"
+#include "ShaderSource.h"
 
 namespace UE::ShaderMinifier
 {
diff --git a/Engine/Source/Developer/VulkanShaderFormat/Private/VulkanShaderCompiler.cpp b/Engine/Source/Developer/VulkanShaderFormat/Private/VulkanShaderCompiler.cpp
--- a/Engine/Source/Developer/VulkanShaderFormat/Private/VulkanShaderCompiler.cpp
+++ b/Engine/Source/Developer/VulkanShaderFormat/Private/VulkanShaderCompiler.cpp
@@ -1760,8 +1760,9 @@ static void VulkanCreateDXCCompileBatchFiles(
 
 // Quick and dirty way to get the location of the entrypoint in the source
 // NOTE: Preprocessed shaders have mcros resolves and comments removed, it makes this easier...
-static FString ParseEntrypointDecl(FStringView PreprocessedShader, FStringView Entrypoint)
+static FString ParseEntrypointDecl(FShaderSource::FViewType PreprocessedShader, FStringView Entrypoint)
 {
+	FShaderSource::FStringType EntrypointConverted(Entrypoint);
 	auto SkipWhitespace = [&](int32& Index)
 	{
 		while (FChar::IsWhitespace(PreprocessedShader[Index]))
@@ -1786,14 +1787,14 @@ static FString ParseEntrypointDecl(FStringView PreprocessedShader, FStringView E
 	FString EntryPointDecl;
 
 	// Go through all the case sensitive matches in the source
-	int32 EntrypointIndex = PreprocessedShader.Find(Entrypoint);
+	int32 EntrypointIndex = PreprocessedShader.Find(EntrypointConverted);
 	check(EntrypointIndex != INDEX_NONE);
 	while (EntrypointIndex != INDEX_NONE)
 	{
 		// This should be the beginning of a new word
 		if ((EntrypointIndex == 0) || !FChar::IsWhitespace(PreprocessedShader[EntrypointIndex - 1]))
 		{
-			EntrypointIndex = PreprocessedShader.Find(Entrypoint, EntrypointIndex + 1);
+			EntrypointIndex = PreprocessedShader.Find(EntrypointConverted, EntrypointIndex + 1);
 			continue;
 		}
 
@@ -1802,38 +1803,38 @@ static FString ParseEntrypointDecl(FStringView PreprocessedShader, FStringView E
 
 		int32 ParamsStart = EntrypointIndex + Entrypoint.Len();
 		SkipWhitespace(ParamsStart);
-		if (PreprocessedShader[ParamsStart] != TEXT('('))
+		if (PreprocessedShader[ParamsStart] != '(')
 		{
-			EntrypointIndex = PreprocessedShader.Find(Entrypoint, ParamsStart);
+			EntrypointIndex = PreprocessedShader.Find(EntrypointConverted, ParamsStart);
 			continue;
 		}
 
-		int32 ParamsEnd = PreprocessedShader.Find(TEXT(")"), ParamsStart+1);
+		int32 ParamsEnd = PreprocessedShader.Find(SHADER_SOURCE_LITERAL(")"), ParamsStart + 1);
 		check(ParamsEnd != INDEX_NONE);
 		if (ParamsEnd == INDEX_NONE)
 		{
 			// Suspicious
-			EntrypointIndex = PreprocessedShader.Find(Entrypoint, ParamsStart);
+			EntrypointIndex = PreprocessedShader.Find(EntrypointConverted, ParamsStart);
 			continue;
 		}
 
 		// Make sure to grab everything up to the function content
 
 		int32 DeclEnd = ParamsEnd + 1;
-		while (PreprocessedShader[DeclEnd] != TEXT('{') && (PreprocessedShader[DeclEnd] != TEXT(';')))
+		while (PreprocessedShader[DeclEnd] != '{' && (PreprocessedShader[DeclEnd] != ';'))
 		{
 			++DeclEnd;
 		}
-		if (PreprocessedShader[DeclEnd] != TEXT('{'))
+		if (PreprocessedShader[DeclEnd] != '{')
 		{
-			EntrypointIndex = PreprocessedShader.Find(Entrypoint, DeclEnd);
+			EntrypointIndex = PreprocessedShader.Find(EntrypointConverted, DeclEnd);
 			continue;
 		}
 
 		// Now back up to pick up the return value, the attributes and everything else that can come with it, like "[numthreads(1,1,1)]"
 
 		int32 DeclBegin = EntrypointIndex - 1;
-		while ( (DeclBegin > 0) && (PreprocessedShader[DeclBegin] != TEXT(';')) && (PreprocessedShader[DeclBegin] != TEXT('}')))
+		while ( (DeclBegin > 0) && (PreprocessedShader[DeclBegin] != ';') && (PreprocessedShader[DeclBegin] != '}'))
 		{
 			--DeclBegin;
 		}
@@ -1850,7 +1851,7 @@ static FString ParseEntrypointDecl(FStringView PreprocessedShader, FStringView E
 
 uint8 ParseWaveSize(
 	const FVulkanShaderCompilerInternalState& InternalState,
-	FStringView PreprocessedShader
+	FShaderSource::FViewType PreprocessedShader
 	)
 {
 	uint8 WaveSize = 0;
@@ -1902,7 +1903,7 @@ uint8 ParseWaveSize(
 
 static bool CompileWithShaderConductor(
 	const FVulkanShaderCompilerInternalState& InternalState,
-	FStringView PreprocessedShader,
+	FShaderSource::FViewType PreprocessedShader,
 	VulkanShaderCompilerSerializedOutput& SerializedOutput,
 	FShaderCompilerOutput&	Output
 )
@@ -2264,7 +2265,7 @@ static void UpdateBindlessUBs(const FVulkanShaderCompilerInternalState& Internal
 
 static bool CompileShaderGroup(
 	FVulkanShaderCompilerInternalState& InternalState,
-	const FString& OriginalPreprocessedShaderSource,
+	const FShaderSource::FStringType& OriginalPreprocessedShaderSource,
 	FShaderCompilerOutput& MergedOutput
 )
 {
@@ -2282,7 +2283,7 @@ static bool CompileShaderGroup(
 		const bool bIsClosestHit = (HitGroupShaderType == FVulkanShaderCompilerInternalState::EHitGroupShaderType::ClosestHit);
 		FShaderCompilerOutput& PartialOutput = bIsClosestHit ? MergedOutput : TempOutput;
 
-		FStringView OrigSourceView(OriginalPreprocessedShaderSource);
+		FShaderSource::FViewType OrigSourceView(OriginalPreprocessedShaderSource);
 		FShaderSource PartialPreprocessedShaderSource(OrigSourceView);
 		UE::ShaderCompilerCommon::RemoveDeadCode(PartialPreprocessedShaderSource, InternalState.GetEntryPointName(), PartialOutput.Errors);
 
@@ -2446,17 +2447,26 @@ void CompileVulkanShader(const FShaderCompilerInput& Input, const FShaderPreproc
 
 	bool bSuccess = false;
 
+#if SHADER_SOURCE_ANSI
+	// Convert to ANSI prior to calling into ShaderConductor. This copy would have been incurred
+	// by SC itself anyways, but would (will?) also be unnecessary if (when) shader parameter parser
+	// is modified to operate on ANSI strings.
+	const FShaderSource::FStringType PreprocessedSourceToCompile(PreprocessedSource);
+#else
+	const FShaderSource::FStringType& PreprocessedSourceToCompile = PreprocessedSource;
+#endif
+
 #if PLATFORM_MAC || PLATFORM_WINDOWS || PLATFORM_LINUX
 	// HitGroup shaders might have multiple entrypoints that we combine into a single blob
 	if (InternalState.HasMultipleEntryPoints())
 	{
-		bSuccess = CompileShaderGroup(InternalState, PreprocessedSource, Output);
+		bSuccess = CompileShaderGroup(InternalState, PreprocessedSourceToCompile, Output);
 	}
 	else
 	{
 		// Compile regular shader via ShaderConductor (DXC)
 		VulkanShaderCompilerSerializedOutput SerializedOutput;
-		bSuccess = CompileWithShaderConductor(InternalState, PreprocessedSource, SerializedOutput, Output);
+		bSuccess = CompileWithShaderConductor(InternalState, PreprocessedSourceToCompile, SerializedOutput, Output);
 
 		if (InternalState.bUseBindlessUniformBuffer)
 		{
diff --git a/Engine/Source/Runtime/RenderCore/Private/ShaderSource.cpp b/Engine/Source/Runtime/RenderCore/Private/ShaderSource.cpp
--- a/Engine/Source/Runtime/RenderCore/Private/ShaderSource.cpp
+++ b/Engine/Source/Runtime/RenderCore/Private/ShaderSource.cpp
@@ -16,12 +16,14 @@ void FShaderSource::Set(FShaderSource::FViewType InSrc, int32 AdditionalSlack)
 	FMemory::Memcpy(SourceData, InSrc.GetData(), sizeof(FShaderSource::CharType) * InSrc.Len());
 }
 
+#if !SHADER_SOURCE_ANSI
 void FShaderSource::Set(FAnsiStringView InSrc)
 {
 	TStringConvert<ANSICHAR, TCHAR> Convert;
 	SetLen(InSrc.Len());
 	Convert.Convert(Source.GetData(), Source.Num(), InSrc.GetData(), InSrc.Len());
 }
+#endif
 
 FShaderSource& FShaderSource::operator=(FShaderSource::FStringType&& InSrc)
 {
diff --git a/Engine/Source/Runtime/RenderCore/Public/ShaderPreprocessTypes.h b/Engine/Source/Runtime/RenderCore/Public/ShaderPreprocessTypes.h
--- a/Engine/Source/Runtime/RenderCore/Public/ShaderPreprocessTypes.h
+++ b/Engine/Source/Runtime/RenderCore/Public/ShaderPreprocessTypes.h
@@ -70,16 +70,28 @@ public:
 
 	FAnsiStringView GetSourceViewAnsi() const
 	{
+#if SHADER_SOURCE_ANSI
+		return PreprocessedSource.GetView();
+#else
 		if (AnsiSource.IsEmpty())
 		{
 			AnsiSource = FAnsiString(PreprocessedSource.GetView());
 		}
 		return FAnsiStringView(AnsiSource);
+#endif
 	}
 
 	FStringView GetSourceViewWide() const
 	{
+#if SHADER_SOURCE_ANSI
+		if (WideSource.IsEmpty())
+		{
+			WideSource = FString(PreprocessedSource.GetView());
+		}
+		return FStringView(WideSource);
+#else
 		return PreprocessedSource.GetView();
+#endif
 	}
 
 	UE_DEPRECATED(5.4, "Use GetUnstrippedSourceView")
@@ -94,7 +106,25 @@ public:
 	{
 		// if the unstripped source is requested, check if the "original source" field has been populated
 		// if not then stripping hasn't occurred so there's only one preprocessed source; return it
+#if SHADER_SOURCE_ANSI
+		if (OriginalPreprocessedSource.IsEmpty())
+		{
+			return FStringView(WideSource);
+		}
+		else
+		{
+			if (WideSourceUnstripped.IsEmpty())
+			{
+				// convert and store wide version of unstripped source if view is requested. 
+				// this is only called in debug paths (a particular case of shader debug dumps,
+				// and visualizing shader source in the editor) so the overhead doesn't matter.
+				WideSourceUnstripped = FString(OriginalPreprocessedSource.GetView());
+			}
+			return FStringView(WideSourceUnstripped);
+		}
+#else
 		return OriginalPreprocessedSource.IsEmpty() ? PreprocessedSource.GetView() : OriginalPreprocessedSource.GetView();
+#endif
 	}
 
 	FShaderSource& EditSource()
@@ -235,8 +265,15 @@ private:
 	// Set by Finalize; original preprocessed source as set by IShaderFormat::PreprocessShader
 	FShaderSource OriginalPreprocessedSource;
 
+#if SHADER_SOURCE_ANSI
+	// Set when GetSourceViewWide or GetUnstrippedSourceView accessors are called (if source is ANSI)
+	// Mutable so we can maintain const correctness in the API otherwise.
+	mutable FString WideSource;
+	mutable FString WideSourceUnstripped;
+#else
 	// Set when GetSourceViewAnsi is called (if source is wide)
 	mutable FAnsiString AnsiSource;
+#endif
 
 	// Array of errors encountered in preprocessing; should be populated by IShaderFormat::PreprocessShader
 	TArray<FShaderCompilerError> Errors;
diff --git a/Engine/Source/Runtime/RenderCore/Public/ShaderSource.h b/Engine/Source/Runtime/RenderCore/Public/ShaderSource.h
--- a/Engine/Source/Runtime/RenderCore/Public/ShaderSource.h
+++ b/Engine/Source/Runtime/RenderCore/Public/ShaderSource.h
@@ -7,8 +7,15 @@
 #include "Containers/StringView.h"
 #include "Containers/UnrealString.h"
 
+#define SHADER_SOURCE_ANSI 1
+
+#if SHADER_SOURCE_ANSI
+#define SHADER_SOURCE_LITERAL(S) S
+#define SHADER_SOURCE_VIEWLITERAL(S) ANSITEXTVIEW(S)
+#else
 #define SHADER_SOURCE_LITERAL(S) TEXT(S)
 #define SHADER_SOURCE_VIEWLITERAL(S) TEXTVIEW(S)
+#endif
 
 /* Class used in shader compilation pipelines which wraps source code and ensures sufficient padding such that 16-byte wide SIMD
  * operations on the source are guaranteed to read valid memory even if starting from the last character.
@@ -23,10 +30,17 @@ private:
 	}
 
 public:
+#if SHADER_SOURCE_ANSI
+	typedef ANSICHAR CharType;
+	typedef FAnsiStringView FViewType;
+	typedef FAnsiString FStringType;
+	typedef FCStringAnsi FCStringType;
+#else
 	typedef TCHAR CharType;
 	typedef FStringView FViewType;
 	typedef FString FStringType;
 	typedef FCString FCStringType;
+#endif
 
 	/** Constexpr predicate indicating whether wide or ansi chars are used */
 	static constexpr bool IsWide() { return sizeof(CharType) == 2; }
@@ -43,7 +57,7 @@ public:
 	 * @param InSrc The source string to be copied
 	 * @param AdditionalSlack optional additional space to allocate; this is on top of the automatic padding.
 	 */
-	RENDERCORE_API FShaderSource(FViewType InSrc, int32 AdditionalSlack = 0);
+	RENDERCORE_API explicit FShaderSource(FViewType InSrc, int32 AdditionalSlack = 0);
 
 	/* Set the given string as the contents of this shader source object. The inner allocation will grow to fit
 	 * the string contents as needed.
@@ -52,12 +66,14 @@ public:
 	 * @param AdditionalSlack optional additional space to allocate; this is on top of the automatic padding.
 	 */
 	RENDERCORE_API void Set(FViewType InSrc, int32 AdditionalSlack = 0);
-
+#if !SHADER_SOURCE_ANSI
 	/* Set the given ANSI string as the contents of this shader source object. The inner allocation will grow to fit
 	 * the string contents as needed. Note that this will incur a conversion to TCHAR and memcpy of the string contents.
+	 * Note that this is only needed if SHADER_SOURCE_ANSI is 0; redundant with Set function accepting a FViewType otherwise.
 	 * @param InSrc The source ANSI string to be copied
 	 */
 	RENDERCORE_API void Set(FAnsiStringView InSrc);
+#endif
 
 	/* Move assignment operator accepting a string object. This will append padding bytes to the existing string, as such it's
 	 * best if there's sufficient extra capacity in the string storage to avoid incurring a realloc-and-copy here.
@@ -102,7 +118,11 @@ public:
 	friend FArchive& operator<<(FArchive& Ar, FShaderSource& ShaderSource);
 	
 private:
-
+#if SHADER_SOURCE_ANSI
+	static constexpr int32 ShaderSourceSimdPadding = 15;
+#else
 	static constexpr int32 ShaderSourceSimdPadding = 7;
+#endif
 	TArray<CharType> Source;
 };
+
