Moving copy-pasted code into ShaderCompilerCommon.

diff --git a/Engine/Source/Developer/Apple/MetalShaderFormat/Private/MetalShaderCompiler.cpp b/Engine/Source/Developer/Apple/MetalShaderFormat/Private/MetalShaderCompiler.cpp
--- a/Engine/Source/Developer/Apple/MetalShaderFormat/Private/MetalShaderCompiler.cpp
+++ b/Engine/Source/Developer/Apple/MetalShaderFormat/Private/MetalShaderCompiler.cpp
@@ -585,16 +585,7 @@ void BuildMetalShaderOutput(
 		BuildResourceTableMapping(ShaderInput.Environment.ResourceTableMap, ShaderInput.Environment.UniformBufferMap, UsedUniformBufferSlots, ShaderOutput.ParameterMap, GenericSRT);
 		CullGlobalUniformBuffers(ShaderInput.Environment.UniformBufferMap, ShaderOutput.ParameterMap);
 
-		// Copy over the bits indicating which resource tables are active.
-		SRT.ResourceTableBits = GenericSRT.ResourceTableBits;
-
-		SRT.ResourceTableLayoutHashes = GenericSRT.ResourceTableLayoutHashes;
-
-		// Now build our token streams.
-		BuildResourceTableTokenStream(GenericSRT.TextureMap, GenericSRT.MaxBoundResourceTable, SRT.TextureMap);
-		BuildResourceTableTokenStream(GenericSRT.ShaderResourceViewMap, GenericSRT.MaxBoundResourceTable, SRT.ShaderResourceViewMap);
-		BuildResourceTableTokenStream(GenericSRT.SamplerMap, GenericSRT.MaxBoundResourceTable, SRT.SamplerMap);
-		BuildResourceTableTokenStream(GenericSRT.UnorderedAccessViewMap, GenericSRT.MaxBoundResourceTable, SRT.UnorderedAccessViewMap);
+		UE::ShaderCompilerCommon::BuildShaderResourceTable(GenericSRT, SRT);
 
 		Header.Bindings.NumUniformBuffers = FMath::Max((uint8)GetNumUniformBuffersUsed(GenericSRT), Header.Bindings.NumUniformBuffers);
 	}
diff --git a/Engine/Source/Developer/ShaderCompilerCommon/Private/ShaderCompilerCommon.cpp b/Engine/Source/Developer/ShaderCompilerCommon/Private/ShaderCompilerCommon.cpp
--- a/Engine/Source/Developer/ShaderCompilerCommon/Private/ShaderCompilerCommon.cpp
+++ b/Engine/Source/Developer/ShaderCompilerCommon/Private/ShaderCompilerCommon.cpp
@@ -83,6 +83,20 @@ void BuildResourceTableTokenStream(const TArray<uint32>& InResourceMap, int32 Ma
 	}
 }
 
+void UE::ShaderCompilerCommon::BuildShaderResourceTable(const FShaderCompilerResourceTable& GenericSRT, FShaderResourceTable& OutSRT, bool bGenerateEmptyTokenStreamIfNoResources)
+{
+	// Copy over the bits indicating which resource tables are active.
+	OutSRT.ResourceTableBits = GenericSRT.ResourceTableBits;
+
+	OutSRT.ResourceTableLayoutHashes = GenericSRT.ResourceTableLayoutHashes;
+
+	// Now build our token streams.
+	BuildResourceTableTokenStream(GenericSRT.TextureMap,             GenericSRT.MaxBoundResourceTable, OutSRT.TextureMap,             bGenerateEmptyTokenStreamIfNoResources);
+	BuildResourceTableTokenStream(GenericSRT.ShaderResourceViewMap,  GenericSRT.MaxBoundResourceTable, OutSRT.ShaderResourceViewMap,  bGenerateEmptyTokenStreamIfNoResources);
+	BuildResourceTableTokenStream(GenericSRT.SamplerMap,             GenericSRT.MaxBoundResourceTable, OutSRT.SamplerMap,             bGenerateEmptyTokenStreamIfNoResources);
+	BuildResourceTableTokenStream(GenericSRT.UnorderedAccessViewMap, GenericSRT.MaxBoundResourceTable, OutSRT.UnorderedAccessViewMap, bGenerateEmptyTokenStreamIfNoResources);
+}
+
 
 bool BuildResourceTableMapping(
 	const FShaderResourceTableMap& ResourceTableMap,
diff --git a/Engine/Source/Developer/ShaderCompilerCommon/Public/ShaderCompilerCommon.h b/Engine/Source/Developer/ShaderCompilerCommon/Public/ShaderCompilerCommon.h
--- a/Engine/Source/Developer/ShaderCompilerCommon/Public/ShaderCompilerCommon.h
+++ b/Engine/Source/Developer/ShaderCompilerCommon/Public/ShaderCompilerCommon.h
@@ -54,6 +54,8 @@ extern SHADERCOMPILERCOMMON_API int16 GetNumUniformBuffersUsed(const FShaderComp
 
 namespace UE::ShaderCompilerCommon
 {
+	extern SHADERCOMPILERCOMMON_API void BuildShaderResourceTable(const FShaderCompilerResourceTable& GenericSRT, FShaderResourceTable& OutSRT, bool bGenerateEmptyTokenStreamIfNoResources = false);
+
 	extern SHADERCOMPILERCOMMON_API bool ExecuteShaderPreprocessingSteps(
 		FShaderPreprocessOutput& PreprocessOutput,
 		const FShaderCompilerInput& Input,
diff --git a/Engine/Source/Developer/ShaderFormatOpenGL/Private/OpenGLShaderCompiler.cpp b/Engine/Source/Developer/ShaderFormatOpenGL/Private/OpenGLShaderCompiler.cpp
--- a/Engine/Source/Developer/ShaderFormatOpenGL/Private/OpenGLShaderCompiler.cpp
+++ b/Engine/Source/Developer/ShaderFormatOpenGL/Private/OpenGLShaderCompiler.cpp
@@ -857,16 +857,7 @@ void BuildShaderOutputInternal(
 		BuildResourceTableMapping(ShaderInput.Environment.ResourceTableMap, ShaderInput.Environment.UniformBufferMap, UsedUniformBufferSlots, ShaderOutput.ParameterMap, GenericSRT);
 		CullGlobalUniformBuffers(ShaderInput.Environment.UniformBufferMap, ShaderOutput.ParameterMap);
 
-		// Copy over the bits indicating which resource tables are active.
-		SRT.ResourceTableBits = GenericSRT.ResourceTableBits;
-
-		SRT.ResourceTableLayoutHashes = GenericSRT.ResourceTableLayoutHashes;
-
-		// Now build our token streams.
-		BuildResourceTableTokenStream(GenericSRT.TextureMap, GenericSRT.MaxBoundResourceTable, SRT.TextureMap);
-		BuildResourceTableTokenStream(GenericSRT.ShaderResourceViewMap, GenericSRT.MaxBoundResourceTable, SRT.ShaderResourceViewMap);
-		BuildResourceTableTokenStream(GenericSRT.SamplerMap, GenericSRT.MaxBoundResourceTable, SRT.SamplerMap);
-		BuildResourceTableTokenStream(GenericSRT.UnorderedAccessViewMap, GenericSRT.MaxBoundResourceTable, SRT.UnorderedAccessViewMap);
+		UE::ShaderCompilerCommon::BuildShaderResourceTable(GenericSRT, SRT);
 	}
 
 	constexpr int32 MaxSamplers = 16;
diff --git a/Engine/Source/Developer/VulkanShaderFormat/Private/VulkanShaderCompiler.cpp b/Engine/Source/Developer/VulkanShaderFormat/Private/VulkanShaderCompiler.cpp
--- a/Engine/Source/Developer/VulkanShaderFormat/Private/VulkanShaderCompiler.cpp
+++ b/Engine/Source/Developer/VulkanShaderFormat/Private/VulkanShaderCompiler.cpp
@@ -472,7 +472,7 @@ static void AddUBResources( const FString& UBName,
 }
 
 static void AddUniformBuffer(
-	const FShaderCompilerResourceTable& ShaderResourceTable,
+	const FShaderResourceTable& ShaderResourceTable,
 	const TArray<VkDescriptorType>& DescriptorTypes,
 	const FShaderCompilerInput& ShaderInput,
 	const FString& UBName,
@@ -589,7 +589,7 @@ static void PrepareUBResourceEntryGlobals(const TArray<uint32>& BindingArray, co
 }
 
 
-static void PrepareGlobals(const FVulkanBindingTable& BindingTable, const FSpirvReflectBindings& SpirvReflectBindings, const FShaderCompilerResourceTable& SRT, const TMap<FString, FVulkanShaderHeader::EType>& EntryTypes, const FShaderCompilerInput& ShaderInput, const TArray<FString>& ParameterNames, FShaderParameterMap& ParameterMap, TArray<FString>& OutGlobalNames, FVulkanShaderHeader& OutHeader)
+static void PrepareGlobals(const FVulkanBindingTable& BindingTable, const FSpirvReflectBindings& SpirvReflectBindings, const FShaderResourceTable& SRT, const TMap<FString, FVulkanShaderHeader::EType>& EntryTypes, const FShaderCompilerInput& ShaderInput, const TArray<FString>& ParameterNames, FShaderParameterMap& ParameterMap, TArray<FString>& OutGlobalNames, FVulkanShaderHeader& OutHeader)
 {
 	auto IsSamplerState = [&SpirvReflectBindings](const FString& ParameterName)
 	{
@@ -711,7 +711,7 @@ static void PrepareGlobals(const FVulkanBindingTable& BindingTable, const FSpirv
 }
 
 static void ConvertToHeader(
-	FShaderCompilerResourceTable& ShaderResourceTable,
+	FShaderResourceTable& ShaderResourceTable,
 	const FVulkanBindingTable& BindingTable,
 	const TArray<VkDescriptorType>& DescriptorTypes,
 	const TMap<FString, FVulkanShaderHeader::EType>& EntryTypes,
@@ -843,21 +843,16 @@ static void ConvertToHeader(
 // NOTE: Keep in sync with BuildResourceTableMapping.
 static FShaderResourceTable BuildSRTFromHeader(const FVulkanShaderHeader& NEWHeader)
 {
-	FShaderResourceTable ShaderResourceTable;
-
-	TArray<uint32> TextureMap;
-	TArray<uint32> ShaderResourceViewMap;
-	TArray<uint32> SamplerMap;
-	TArray<uint32> UnorderedAccessViewMap;
+	FShaderCompilerResourceTable GenericSRT;
 
 	for (int32 UBIndex = 0; UBIndex < NEWHeader.UniformBuffers.Num(); ++UBIndex)
 	{
 		const FVulkanShaderHeader::FUniformBufferInfo& UBHeader = NEWHeader.UniformBuffers[UBIndex];
 
-		ShaderResourceTable.ResourceTableLayoutHashes.Emplace(UBHeader.LayoutHash);
+		GenericSRT.ResourceTableLayoutHashes.Emplace(UBHeader.LayoutHash);
 		if (UBHeader.ResourceEntries.Num() > 0)
 		{
-			ShaderResourceTable.ResourceTableBits |= 1 << UBIndex;
+			GenericSRT.ResourceTableBits |= 1 << UBIndex;
 
 			for (const FVulkanShaderHeader::FUBResourceInfo& UBRes : UBHeader.ResourceEntries)
 			{
@@ -866,20 +861,20 @@ static FShaderResourceTable BuildSRTFromHeader(const FVulkanShaderHeader& NEWHea
 				{
 				case UBMT_TEXTURE:
 				case UBMT_RDG_TEXTURE:
-					TextureMap.Add(ResourceMap);
+					GenericSRT.TextureMap.Add(ResourceMap);
 					break;
 				case UBMT_SAMPLER:
-					SamplerMap.Add(ResourceMap);
+					GenericSRT.SamplerMap.Add(ResourceMap);
 					break;
 				case UBMT_SRV:
 				case UBMT_RDG_TEXTURE_SRV:
 				case UBMT_RDG_BUFFER_SRV:
-					ShaderResourceViewMap.Add(ResourceMap);
+					GenericSRT.ShaderResourceViewMap.Add(ResourceMap);
 					break;
 				case UBMT_UAV:
 				case UBMT_RDG_TEXTURE_UAV:
 				case UBMT_RDG_BUFFER_UAV:
-					UnorderedAccessViewMap.Add(ResourceMap);
+					GenericSRT.UnorderedAccessViewMap.Add(ResourceMap);
 					break;
 				default:
 					check(false);
@@ -888,11 +883,10 @@ static FShaderResourceTable BuildSRTFromHeader(const FVulkanShaderHeader& NEWHea
 		}
 	}
 
-	const int32 MaxBoundResourceTable = NEWHeader.UniformBuffers.Num();
-	BuildResourceTableTokenStream(TextureMap, MaxBoundResourceTable, ShaderResourceTable.TextureMap);
-	BuildResourceTableTokenStream(ShaderResourceViewMap, MaxBoundResourceTable, ShaderResourceTable.ShaderResourceViewMap);
-	BuildResourceTableTokenStream(SamplerMap, MaxBoundResourceTable, ShaderResourceTable.SamplerMap);
-	BuildResourceTableTokenStream(UnorderedAccessViewMap, MaxBoundResourceTable, ShaderResourceTable.UnorderedAccessViewMap);
+	GenericSRT.MaxBoundResourceTable = NEWHeader.UniformBuffers.Num();
+
+	FShaderResourceTable ShaderResourceTable;
+	UE::ShaderCompilerCommon::BuildShaderResourceTable(GenericSRT, ShaderResourceTable);
 
 	return ShaderResourceTable;
 }
@@ -1018,7 +1012,7 @@ static void BuildShaderOutput(
 	ShaderOutput.ParameterMap.GetAllParameterNames(OriginalParameters);
 
 	// Build the SRT for this shader.
-	FShaderCompilerResourceTable ShaderResourceTable;
+	FShaderResourceTable ShaderResourceTable;
 	{
 		// Build the generic SRT for this shader.
 		FShaderCompilerResourceTable GenericSRT;
@@ -1028,15 +1022,7 @@ static void BuildShaderOutput(
 			return;
 		}
 
-		// Copy over the bits indicating which resource tables are active.
-		ShaderResourceTable.ResourceTableBits = GenericSRT.ResourceTableBits;
-		ShaderResourceTable.ResourceTableLayoutHashes = GenericSRT.ResourceTableLayoutHashes;
-
-		// Now build our token streams.
-		BuildResourceTableTokenStream(GenericSRT.TextureMap, GenericSRT.MaxBoundResourceTable, ShaderResourceTable.TextureMap, true);
-		BuildResourceTableTokenStream(GenericSRT.ShaderResourceViewMap, GenericSRT.MaxBoundResourceTable, ShaderResourceTable.ShaderResourceViewMap, true);
-		BuildResourceTableTokenStream(GenericSRT.SamplerMap, GenericSRT.MaxBoundResourceTable, ShaderResourceTable.SamplerMap, true);
-		BuildResourceTableTokenStream(GenericSRT.UnorderedAccessViewMap, GenericSRT.MaxBoundResourceTable, ShaderResourceTable.UnorderedAccessViewMap, true);
+		UE::ShaderCompilerCommon::BuildShaderResourceTable(GenericSRT, ShaderResourceTable, true);
 	}
 
 	TArray<FString> NewParameters;
diff --git a/Engine/Source/Developer/Windows/ShaderFormatD3D/Private/D3DShaderCompiler.inl b/Engine/Source/Developer/Windows/ShaderFormatD3D/Private/D3DShaderCompiler.inl
--- a/Engine/Source/Developer/Windows/ShaderFormatD3D/Private/D3DShaderCompiler.inl
+++ b/Engine/Source/Developer/Windows/ShaderFormatD3D/Private/D3DShaderCompiler.inl
@@ -438,16 +438,7 @@ inline void GenerateFinalOutput(
 		FMemoryWriter UniformBufferNameWriter(UniformBufferNameBytes);
 		UniformBufferNameWriter << CompileData.UniformBufferNames;
 
-		// Copy over the bits indicating which resource tables are active.
-		SRT.ResourceTableBits = GenericSRT.ResourceTableBits;
-
-		SRT.ResourceTableLayoutHashes = GenericSRT.ResourceTableLayoutHashes;
-
-		// Now build our token streams.
-		BuildResourceTableTokenStream(GenericSRT.TextureMap, GenericSRT.MaxBoundResourceTable, SRT.TextureMap);
-		BuildResourceTableTokenStream(GenericSRT.ShaderResourceViewMap, GenericSRT.MaxBoundResourceTable, SRT.ShaderResourceViewMap);
-		BuildResourceTableTokenStream(GenericSRT.SamplerMap, GenericSRT.MaxBoundResourceTable, SRT.SamplerMap);
-		BuildResourceTableTokenStream(GenericSRT.UnorderedAccessViewMap, GenericSRT.MaxBoundResourceTable, SRT.UnorderedAccessViewMap);
+		UE::ShaderCompilerCommon::BuildShaderResourceTable(GenericSRT, SRT);
 	}
 
 	if (Input.Environment.CompilerFlags.Contains(CFLAG_ForceRemoveUnusedInterpolators) && Input.Target.Frequency == SF_Pixel && Input.bCompilingForShaderPipeline && bProcessingSecondTime)
