VulkanRHI: Fix Vulkan bindless shader compilation

#rnx

diff --git a/Engine/Source/Developer/VulkanShaderFormat/Private/VulkanShaderCompiler.cpp b/Engine/Source/Developer/VulkanShaderFormat/Private/VulkanShaderCompiler.cpp
--- a/Engine/Source/Developer/VulkanShaderFormat/Private/VulkanShaderCompiler.cpp
+++ b/Engine/Source/Developer/VulkanShaderFormat/Private/VulkanShaderCompiler.cpp
@@ -1904,8 +1904,14 @@ static bool CompileWithShaderConductor(
 	FShaderCompilerDefinitions AdditionalDefines;
 	PRAGMA_ENABLE_DEPRECATION_WARNINGS
 
+	TArray<FString> ExtraDxcArgs;
+	if (InternalState.IsSM6())
+	{
+		ExtraDxcArgs.Add(TEXT("-fvk-allow-rwstructuredbuffer-arrays"));
+	}
+
 	// Load shader source into compiler context
-	CompilerContext.LoadSource(PreprocessedShader, Input.VirtualSourceFilePath, InternalState.GetEntryPointName(), InternalState.GetShaderFrequency(), &AdditionalDefines);
+	CompilerContext.LoadSource(PreprocessedShader, Input.VirtualSourceFilePath, InternalState.GetEntryPointName(), InternalState.GetShaderFrequency(), &AdditionalDefines, &ExtraDxcArgs);
 
 	// Initialize compilation options for ShaderConductor
 	CrossCompiler::FShaderConductorOptions Options;
