Unifying bindless code generation.
- Array-based bindless (Vulkan) only needed one more piece of data to use "GenerateBindlessAccess"
- Now all bindless shader platforms will generate a getter function for each bindless resource/sampler.

diff --git a/Engine/Shaders/Public/Platform/Vulkan/VulkanCommon.ush b/Engine/Shaders/Public/Platform/Vulkan/VulkanCommon.ush
--- a/Engine/Shaders/Public/Platform/Vulkan/VulkanCommon.ush
+++ b/Engine/Shaders/Public/Platform/Vulkan/VulkanCommon.ush
@@ -42,10 +42,11 @@ void ImageInterlockedAddUInt64(RWTexture2D<UlongType> Dest, uint2 Coord, UlongTy
 
 #if PLATFORM_SUPPORTS_BINDLESS
 	// While we use descriptor indexing, declare unbounded resource arrays
-	COMPILER_DEFINE VULKAN_HEAP_NAME(Prefix, Type) Prefix##Type
-	COMPILER_DEFINE VULKAN_SAMPLER_HEAP(Type) VULKAN_HEAP_NAME(VULKAN_BINDLESS_SAMPLER_ARRAY_PREFIX,Type)
-	COMPILER_DEFINE VULKAN_SRV_HEAP(Type) VULKAN_HEAP_NAME(VULKAN_BINDLESS_SRV_ARRAY_PREFIX,Type)
-	COMPILER_DEFINE VULKAN_UAV_HEAP(Type) VULKAN_HEAP_NAME(VULKAN_BINDLESS_UAV_ARRAY_PREFIX,Type)
+	#define VULKAN_HEAP_NAME_INNER(Prefix, Type) Prefix##Type
+	#define VULKAN_HEAP_NAME(Prefix, Type) VULKAN_HEAP_NAME_INNER(Prefix, Type)
+	#define VULKAN_SAMPLER_HEAP(Type) VULKAN_HEAP_NAME(BINDLESS_SAMPLER_ARRAY_PREFIX,Type)
+	#define VULKAN_SRV_HEAP(Type) VULKAN_HEAP_NAME(BINDLESS_SRV_ARRAY_PREFIX,Type)
+	#define VULKAN_UAV_HEAP(Type) VULKAN_HEAP_NAME(BINDLESS_UAV_ARRAY_PREFIX,Type)
 
 	#define UB_RESOURCE_MEMBER_SRV(Type, UBName, Name) \
 		typedef Type SafeType##UBName##_##Name; \
diff --git a/Engine/Source/Developer/Apple/MetalShaderFormat/Private/MetalCompileShaderMSC.cpp b/Engine/Source/Developer/Apple/MetalShaderFormat/Private/MetalCompileShaderMSC.cpp
--- a/Engine/Source/Developer/Apple/MetalShaderFormat/Private/MetalCompileShaderMSC.cpp
+++ b/Engine/Source/Developer/Apple/MetalShaderFormat/Private/MetalCompileShaderMSC.cpp
@@ -411,7 +411,7 @@ struct FMetalShaderParameterParserPlatformConfiguration : public FShaderParamete
 	{
 	}
 
-	virtual FString GenerateBindlessAccess(EBindlessConversionType BindlessType, FStringView ShaderTypeString, FStringView IndexString) const final
+	virtual FString GenerateBindlessAccess(EBindlessConversionType BindlessType, FStringView FullTypeString, FStringView ArrayNameOverride, FStringView IndexString) const final
 	{
 		// GetResourceFromHeap(Type, Index) ResourceDescriptorHeap[Index]
 		// GetSamplerFromHeap(Type, Index)  SamplerDescriptorHeap[Index]
diff --git a/Engine/Source/Developer/VulkanShaderFormat/Private/VulkanShaderCompiler.cpp b/Engine/Source/Developer/VulkanShaderFormat/Private/VulkanShaderCompiler.cpp
--- a/Engine/Source/Developer/VulkanShaderFormat/Private/VulkanShaderCompiler.cpp
+++ b/Engine/Source/Developer/VulkanShaderFormat/Private/VulkanShaderCompiler.cpp
@@ -2077,9 +2077,9 @@ void ModifyVulkanCompilerInput(FShaderCompilerInput& Input)
 		Input.Environment.SetDefine(TEXT("VULKAN_SUPPORTS_SUBGROUP_SIZE_CONTROL"), 1);
 	}
 
-	Input.Environment.SetDefine(TEXT("VULKAN_BINDLESS_SRV_ARRAY_PREFIX"), FShaderParameterParser::kBindlessSRVArrayPrefix);
-	Input.Environment.SetDefine(TEXT("VULKAN_BINDLESS_UAV_ARRAY_PREFIX"), FShaderParameterParser::kBindlessUAVArrayPrefix);
-	Input.Environment.SetDefine(TEXT("VULKAN_BINDLESS_SAMPLER_ARRAY_PREFIX"), FShaderParameterParser::kBindlessSamplerArrayPrefix);
+	Input.Environment.SetDefine(TEXT("BINDLESS_SRV_ARRAY_PREFIX"), FShaderParameterParser::kBindlessSRVArrayPrefix);
+	Input.Environment.SetDefine(TEXT("BINDLESS_UAV_ARRAY_PREFIX"), FShaderParameterParser::kBindlessUAVArrayPrefix);
+	Input.Environment.SetDefine(TEXT("BINDLESS_SAMPLER_ARRAY_PREFIX"), FShaderParameterParser::kBindlessSamplerArrayPrefix);
 	Input.Environment.SetDefine(TEXT("VULKAN_MAX_BINDLESS_UNIFORM_BUFFERS_PER_STAGE"), VulkanBindless::MaxUniformBuffersPerStage);
 
 	if (IsAndroidShaderFormat(Input.ShaderFormat))
@@ -2362,9 +2362,9 @@ static bool CompileShaderGroup(
 	return bSuccess;
 }
 
-struct FPS5ShaderParameterParserPlatformConfiguration : public FShaderParameterParser::FPlatformConfiguration
+struct FVulkanShaderParameterParserPlatformConfiguration : public FShaderParameterParser::FPlatformConfiguration
 {
-	FPS5ShaderParameterParserPlatformConfiguration(const FShaderCompilerInput& Input)
+	FVulkanShaderParameterParserPlatformConfiguration(const FShaderCompilerInput& Input)
 		: FShaderParameterParser::FPlatformConfiguration()
 	{
 		EnumAddFlags(Flags, EShaderParameterParserConfigurationFlags::SupportsBindless | EShaderParameterParserConfigurationFlags::BindlessUsesArrays);
@@ -2377,10 +2377,12 @@ struct FPS5ShaderParameterParserPlatformConfiguration : public FShaderParameterP
 		}
 	}
 
-	virtual FString GenerateBindlessAccess(EBindlessConversionType BindlessType, FStringView ShaderTypeString, FStringView IndexString) const final
+	virtual FString GenerateBindlessAccess(EBindlessConversionType BindlessType, FStringView FullTypeString, FStringView ArrayNameOverride, FStringView IndexString) const final
 	{
-		checkf(false, TEXT("Vulkan does not use GenerateBindlessAccess"));
-		return FString();
+		// Heap[Index]
+		return FString::Printf(TEXT("%.*s[%.*s]"),
+			ArrayNameOverride.Len(), ArrayNameOverride.GetData(),
+			IndexString.Len(), IndexString.GetData());
 	}
 };
 
@@ -2391,7 +2393,7 @@ void CompileVulkanShader(const FShaderCompilerInput& Input, const FShaderPreproc
 	FString EntryPointName = Input.EntryPointName;
 	FString PreprocessedSource(InPreprocessOutput.GetSourceViewWide());
 
-	FPS5ShaderParameterParserPlatformConfiguration PlatformConfiguration(Input);
+	FVulkanShaderParameterParserPlatformConfiguration PlatformConfiguration(Input);
 	FShaderParameterParser ShaderParameterParser(PlatformConfiguration);
 	if (!ShaderParameterParser.ParseAndModify(Input, Output.Errors, PreprocessedSource))
 	{
diff --git a/Engine/Source/Developer/Windows/ShaderFormatD3D/Private/D3DShaderCompiler.cpp b/Engine/Source/Developer/Windows/ShaderFormatD3D/Private/D3DShaderCompiler.cpp
--- a/Engine/Source/Developer/Windows/ShaderFormatD3D/Private/D3DShaderCompiler.cpp
+++ b/Engine/Source/Developer/Windows/ShaderFormatD3D/Private/D3DShaderCompiler.cpp
@@ -1232,7 +1232,7 @@ struct FD3DShaderParameterParserPlatformConfiguration : public FShaderParameterP
 	{
 	}
 
-	virtual FString GenerateBindlessAccess(EBindlessConversionType BindlessType, FStringView ShaderTypeString, FStringView IndexString) const final
+	virtual FString GenerateBindlessAccess(EBindlessConversionType BindlessType, FStringView FullTypeString, FStringView ArrayNameOverride, FStringView IndexString) const final
 	{
 		// GetResourceFromHeap(Type, Index) ResourceDescriptorHeap[Index]
 		// GetSamplerFromHeap(Type, Index)  SamplerDescriptorHeap[Index]
diff --git a/Engine/Source/Runtime/RenderCore/Private/ShaderParameterParser.cpp b/Engine/Source/Runtime/RenderCore/Private/ShaderParameterParser.cpp
--- a/Engine/Source/Runtime/RenderCore/Private/ShaderParameterParser.cpp
+++ b/Engine/Source/Runtime/RenderCore/Private/ShaderParameterParser.cpp
@@ -956,29 +956,27 @@ FString FShaderParameterParser::GenerateBindlessParameterDeclaration(const FPars
 	TStringBuilder<64> FullType;
 	FullType << StorageClass << TypedefName;
 
+	TStringBuilder<64> HeapName;
+
 	if (EnumHasAnyFlags(PlatformConfiguration.Flags, EShaderParameterParserConfigurationFlags::BindlessUsesArrays))
 	{
 		const FStringView HeapPrefix = GetBindlessArrayHeapPrefix(ParsedParameter.BindlessConversionType);
 
+		HeapName << HeapPrefix << TypedefName;
+
 		// Declare a heap for the RewriteType
 		// e.g. `SafeType##Name ResourceDescriptorHeap_SafeType##Name[];`
-		Result << FullType << TEXT(" ") << HeapPrefix << TypedefName << TEXT("[]; ");
+		Result << FullType << TEXT(" ") << HeapName << TEXT("[]; ");
 		// :todo-jn: specify the descriptor set and binding directly in source instead of patching SPIRV
-
-		// e.g. `static const SafeType##Name Name = ResourceDescriptorHeap_SafeType##Name[BindlessResource_##Name];`
-		Result << TEXT("static const ") << FullType << TEXT(" ") << Name << TEXT(" = ") << HeapPrefix << TypedefName << TEXT("[") << IndexString << TEXT("];");
 	}
-	else
-	{
-		const FString BindlessAccess = PlatformConfiguration.GenerateBindlessAccess(ParsedParameter.BindlessConversionType, FullType, IndexString);
 
-		const TCHAR* Kind = bIsSampler ? TEXT("Sampler") : TEXT("Resource");
+	const FString BindlessAccess = PlatformConfiguration.GenerateBindlessAccess(ParsedParameter.BindlessConversionType, FullType, HeapName, IndexString);
+	const TCHAR* FunctionPrefix = bIsSampler ? TEXT("GetBindlessSampler") : TEXT("GetBindlessResource");
 
-		// e.g. `Type GetBindlessResource##Name() { return GetResourceFromHeap(Type, BindlessResource_##Name); } static const Type Name = GetBindlessResource##Name()`
-		// or   `Type GetBindlessSampler##Name() { return GetSamplerFromHeap(Type, BindlessSampler_##Name); } static const Type Name = GetBindlessSampler##Name()`
-		Result << FullType << TEXT(" GetBindless") << Kind << Name << TEXT("() { return ") << BindlessAccess << TEXT("; } ");
-		Result << TEXT("static const ") << FullType << TEXT(" ") << Name << TEXT(" = GetBindless") << Kind << Name << TEXT("();");
-	}
+	// e.g. `Type GetBindlessResource##Name() { return GetResourceFromHeap(Type, BindlessResource_##Name); } static const Type Name = GetBindlessResource##Name()`
+	// or   `Type GetBindlessSampler##Name() { return GetSamplerFromHeap(Type, BindlessSampler_##Name); } static const Type Name = GetBindlessSampler##Name()`
+	Result << FullType << TEXT(" ") << FunctionPrefix << Name << TEXT("() { return ") << BindlessAccess << TEXT("; } ");
+	Result << TEXT("static const ") << FullType << TEXT(" ") << Name << TEXT(" = ") << FunctionPrefix << Name << TEXT("();");
 
 	return Result.ToString();
 }
diff --git a/Engine/Source/Runtime/RenderCore/Public/ShaderParameterParser.h b/Engine/Source/Runtime/RenderCore/Public/ShaderParameterParser.h
--- a/Engine/Source/Runtime/RenderCore/Public/ShaderParameterParser.h
+++ b/Engine/Source/Runtime/RenderCore/Public/ShaderParameterParser.h
@@ -50,7 +50,7 @@ public:
 		}
 
 		/** Generate shader code for accessing a bindless resource or sampler */
-		virtual FString GenerateBindlessAccess(EBindlessConversionType BindlessType, FStringView ShaderTypeString, FStringView IndexString) const
+		virtual FString GenerateBindlessAccess(EBindlessConversionType BindlessType, FStringView FullTypeString, FStringView ArrayNameOverride, FStringView IndexString) const
 		{
 			checkf(false, TEXT("Platforms that support bindless must override GenerateBindlessAccess"));
 			return FString();
