VulkanRHI: Fix some  errors in Vulkan ray tracing shader compilation.

diff --git a/Engine/Source/Developer/VulkanShaderFormat/Private/VulkanShaderCompiler.cpp b/Engine/Source/Developer/VulkanShaderFormat/Private/VulkanShaderCompiler.cpp
--- a/Engine/Source/Developer/VulkanShaderFormat/Private/VulkanShaderCompiler.cpp
+++ b/Engine/Source/Developer/VulkanShaderFormat/Private/VulkanShaderCompiler.cpp
@@ -1924,6 +1924,13 @@ static bool CompileWithShaderConductor(
 	// VK_EXT_scalar_block_layout is required by raytracing and by Nanite (so expect it to be present in SM6/Vulkan_1_3)
 	Options.bDisableScalarBlockLayout = !(InternalState.IsRayTracingShader() || InternalState.IsSM6());
 
+	if (InternalState.IsRayTracingShader() || InternalState.IsSM6())
+	{
+		// Use SM 6.6 as the baseline for Vulkan SM6 shaders
+		Options.ShaderModel.Major = 6;
+		Options.ShaderModel.Minor = 6;
+	}
+
 	if (Input.Environment.CompilerFlags.Contains(CFLAG_AllowRealTypes))
 	{
 		Options.bEnable16bitTypes = true;
