VulkanRHI: fix Vulkan ray tracing to use indices in shader record for bindless HitGroupSystemIndexBuffer and HitGroupSystemVertexBuffer.

#rnx

diff --git a/Engine/Shaders/Public/Platform/Vulkan/VulkanCommon.ush b/Engine/Shaders/Public/Platform/Vulkan/VulkanCommon.ush
--- a/Engine/Shaders/Public/Platform/Vulkan/VulkanCommon.ush
+++ b/Engine/Shaders/Public/Platform/Vulkan/VulkanCommon.ush
@@ -102,7 +102,8 @@ void ImageInterlockedAddUInt64(RWTexture2D<UlongType> Dest, uint2 Coord, UlongTy
 
 	static const FHitGroupSystemRootConstants HitGroupSystemRootConstants = VulkanHitGroupSystemParameters.RootConstants;
 
-	ByteAddressBuffer VULKAN_SRV_HEAP(HitGroupSystemBuffers)[];
-	static const ByteAddressBuffer HitGroupSystemIndexBuffer = VULKAN_SRV_HEAP(HitGroupSystemBuffers)[VulkanHitGroupSystemParameters.BindlessHitGroupSystemIndexBuffer];
-	static const ByteAddressBuffer HitGroupSystemVertexBuffer = VULKAN_SRV_HEAP(HitGroupSystemBuffers)[VulkanHitGroupSystemParameters.BindlessHitGroupSystemVertexBuffer];
+	// Will get filled in with the bindless index in VulkanHitGroupSystemParameters.BindlessHitGroupSystemIndexBuffer
+	ByteAddressBuffer HitGroupSystemIndexBuffer;
+	// Will get filled in with the bindless index in VulkanHitGroupSystemParameters.BindlessHitGroupSystemVertexBuffer
+	ByteAddressBuffer HitGroupSystemVertexBuffer;
 #endif
diff --git a/Engine/Source/Developer/VulkanShaderFormat/Private/VulkanShaderCompiler.cpp b/Engine/Source/Developer/VulkanShaderFormat/Private/VulkanShaderCompiler.cpp
--- a/Engine/Source/Developer/VulkanShaderFormat/Private/VulkanShaderCompiler.cpp
+++ b/Engine/Source/Developer/VulkanShaderFormat/Private/VulkanShaderCompiler.cpp
@@ -2356,6 +2356,9 @@ struct FVulkanShaderParameterParserPlatformConfiguration : public FShaderParamet
 {
 	FVulkanShaderParameterParserPlatformConfiguration(const FShaderCompilerInput& Input)
 		: FShaderParameterParser::FPlatformConfiguration()
+		, bIsRayTracingShader(Input.IsRayTracingShader())
+		, HitGroupSystemIndexBufferName(FShaderParameterParser::kBindlessSRVPrefix + FString(TEXT("HitGroupSystemIndexBuffer")))
+		, HitGroupSystemVertexBufferName(FShaderParameterParser::kBindlessSRVPrefix + FString(TEXT("HitGroupSystemVertexBuffer")))
 	{
 		EnumAddFlags(Flags, EShaderParameterParserConfigurationFlags::SupportsBindless | EShaderParameterParserConfigurationFlags::BindlessUsesArrays);
 
@@ -2369,11 +2372,28 @@ struct FVulkanShaderParameterParserPlatformConfiguration : public FShaderParamet
 
 	virtual FString GenerateBindlessAccess(EBindlessConversionType BindlessType, FStringView FullTypeString, FStringView ArrayNameOverride, FStringView IndexString) const final
 	{
+		if (bIsRayTracingShader && (BindlessType == EBindlessConversionType::SRV))
+		{
+			// Patch the HitGroupSystemIndexBuffer/HitGroupSystemVertexBuffer indices to use the ones contained in the shader record
+			if (IndexString == HitGroupSystemIndexBufferName)
+			{
+				IndexString = TEXTVIEW("VulkanHitGroupSystemParameters.BindlessHitGroupSystemIndexBuffer");
+			}
+			else if (IndexString == HitGroupSystemVertexBufferName)
+			{
+				IndexString = TEXTVIEW("VulkanHitGroupSystemParameters.BindlessHitGroupSystemVertexBuffer");
+			}
+		}
+
 		// Heap[Index]
 		return FString::Printf(TEXT("%.*s[%.*s]"),
 			ArrayNameOverride.Len(), ArrayNameOverride.GetData(),
 			IndexString.Len(), IndexString.GetData());
 	}
+
+	const bool bIsRayTracingShader;
+	const FString HitGroupSystemIndexBufferName;
+	const FString HitGroupSystemVertexBufferName;
 };
 
 void CompileVulkanShader(const FShaderCompilerInput& Input, const FShaderPreprocessOutput& InPreprocessOutput, FShaderCompilerOutput& Output, const class FString& WorkingDirectory)
diff --git a/Engine/Source/Runtime/RenderCore/Public/ShaderParameterParser.h b/Engine/Source/Runtime/RenderCore/Public/ShaderParameterParser.h
--- a/Engine/Source/Runtime/RenderCore/Public/ShaderParameterParser.h
+++ b/Engine/Source/Runtime/RenderCore/Public/ShaderParameterParser.h
@@ -49,6 +49,9 @@ public:
 			}
 		}
 
+		virtual ~FPlatformConfiguration()
+		{}
+
 		/** Generate shader code for accessing a bindless resource or sampler */
 		virtual FString GenerateBindlessAccess(EBindlessConversionType BindlessType, FStringView FullTypeString, FStringView ArrayNameOverride, FStringView IndexString) const
 		{
