VulkanRHI: Always remove dead code from Vulkan ray tracing shaders (necessary to remove the unused entrypoints).

#rnx

diff --git a/Engine/Source/Developer/VulkanShaderFormat/Private/VulkanShaderCompiler.cpp b/Engine/Source/Developer/VulkanShaderFormat/Private/VulkanShaderCompiler.cpp
--- a/Engine/Source/Developer/VulkanShaderFormat/Private/VulkanShaderCompiler.cpp
+++ b/Engine/Source/Developer/VulkanShaderFormat/Private/VulkanShaderCompiler.cpp
@@ -2082,6 +2082,10 @@ void ModifyVulkanCompilerInput(FShaderCompilerInput& Input)
 	{
 		// Name of the structure in raytracing shader records in VulkanCommon.usf
 		Input.RequiredSymbols.Add(TEXT("HitGroupSystemRootConstants"));
+
+		// Always remove dead code for ray tracing shaders regardless of cvar settings, 
+		// we can't support multiple entrypoints remaining in the binaries
+		Input.Environment.CompilerFlags.Add(CFLAG_RemoveDeadCode);
 	}
 }
 
