VulkanRHI: Clean up shader structures.  Remove unused FVulkanBindingTable and EVulkanBindingType.  Adding ShaderStage::EStage::NumGraphicsStages to eventually replace ::NumStages where appropriate.  Start simplifying VulkanShaderCompiler to remove some back and forth caused by legacy code paths. (WIP)

#rnx

diff --git a/Engine/Shaders/Public/Platform/Vulkan/VulkanCommon.ush b/Engine/Shaders/Public/Platform/Vulkan/VulkanCommon.ush
--- a/Engine/Shaders/Public/Platform/Vulkan/VulkanCommon.ush
+++ b/Engine/Shaders/Public/Platform/Vulkan/VulkanCommon.ush
@@ -7,7 +7,7 @@
 #pragma once
 
 // Update this GUID to invalidate shader recompilation for only Vulkan shaders
-#pragma message("UESHADERMETADATA_VERSION BB707B6A-C217-4B23-9A22-086E9EC77FCF")
+#pragma message("UESHADERMETADATA_VERSION 8CDED3E3-2E30-4B46-A493-70E664778036")
 
 #define COMPILER_SUPPORTS_ULONG_TYPES 1
 #define COMPILER_SUPPORTS_UINT64_IMAGE_ATOMICS (1 && COMPILER_SUPPORTS_ULONG_TYPES)
diff --git a/Engine/Source/Developer/VulkanShaderFormat/Private/VulkanBackend.cpp b/Engine/Source/Developer/VulkanShaderFormat/Private/VulkanBackend.cpp
--- a/Engine/Source/Developer/VulkanShaderFormat/Private/VulkanBackend.cpp
+++ b/Engine/Source/Developer/VulkanShaderFormat/Private/VulkanBackend.cpp
@@ -74,77 +74,3 @@ const TCHAR* VULKAN_SUBPASS_FETCH_VAR_W[9] =
 	TEXT("GENERATED_SubpassFetchAttachment6"),
 	TEXT("GENERATED_SubpassFetchAttachment7")
 };
-
-FVulkanBindingTable::FBinding::FBinding()
-{
-	FMemory::Memzero(Name);
-}
-
-FVulkanBindingTable::FBinding::FBinding(const char* InName, int32 InVirtualIndex, EVulkanBindingType::EType InType, int8 InSubType) :
-	VirtualIndex(InVirtualIndex),
-	Type(InType),
-	SubType(InSubType)
-{
-	check(InName);
-	int32 NewNameLength = sizeof(char) * (strlen(InName) + 1);
-	check(NewNameLength < sizeof(Name));
-	FMemory::Memcpy(Name, InName, NewNameLength);
-
-	// Validate Sampler type, s == PACKED_TYPENAME_SAMPLER
-	check((Type == EVulkanBindingType::CombinedImageSampler || Type == EVulkanBindingType::UniformTexelBuffer) ? SubType == 's' : true);
-
-	check(Type != EVulkanBindingType::PackedUniformBuffer || CrossCompiler::IsValidPackedTypeName((CrossCompiler::EPackedTypeName)SubType));
-}
-
-inline int8 ExtractHLSLCCType(const char* name)
-{
-	check(name);
-
-	#pragma warning( push )
-	#pragma warning( disable: 4996 )
-	int32 len = strlen(name);
-	#pragma warning( pop )
-	
-	check(len > 0);
-	int8 TypeChar = name[len-1];
-	return TypeChar;
-}
-
-int32 FVulkanBindingTable::RegisterBinding(const char* InName, const char* BlockName, EVulkanBindingType::EType Type)
-{
-	check(InName);
-	if (!*InName)
-	{
-		return -1;
-	}
-
-	for (int32 Index = 0; Index < Bindings.Num(); ++Index)
-	{
-		if (FCStringAnsi::Strcmp(Bindings[Index].Name, InName) == 0)
-		{
-			return Index;
-		}
-	}
-
-	int32 BindingIdx = Bindings.Num();
-	
-	Bindings.Add(FBinding(InName, BindingIdx, Type, ExtractHLSLCCType(BlockName)));
-
-	return BindingIdx;
-}
-
-void FVulkanBindingTable::SortBindings()
-{
-	// Order is guaranteed to match EVulkanBindingType::EType
-	check(!bSorted);
-	Bindings.Sort([](const FBinding& A, const FBinding& B)
-		{
-			if (A.Type == B.Type)
-			{
-				return A.VirtualIndex < B.VirtualIndex;
-			}
-
-			return A.Type < B.Type;
-		});
-	bSorted = true;
-}
diff --git a/Engine/Source/Developer/VulkanShaderFormat/Private/VulkanBackend.h b/Engine/Source/Developer/VulkanShaderFormat/Private/VulkanBackend.h
--- a/Engine/Source/Developer/VulkanShaderFormat/Private/VulkanBackend.h
+++ b/Engine/Source/Developer/VulkanShaderFormat/Private/VulkanBackend.h
@@ -12,56 +12,6 @@
 #pragma GCC visibility push(default)
 #endif // __GNUC__
 
-struct FVulkanBindingTable
-{
-	struct FBinding
-	{
-		FBinding();
-		FBinding(const char* InName, int32 InVirtualIndex, EVulkanBindingType::EType InType, int8 InSubType);
-
-		char		Name[256];
-		int32		VirtualIndex = -1;
-		EVulkanBindingType::EType	Type;
-		int8		SubType;	// HLSL CC subtype, PACKED_TYPENAME_HIGHP and etc
-	};
-		
-	FVulkanBindingTable(EHlslShaderFrequency ShaderStage) : Stage(ShaderStage) {}
-
-	int32 RegisterBinding(const char* InName, const char* BlockName, EVulkanBindingType::EType Type);
-
-	const TArray<FBinding>& GetBindings() const
-	{
-		check(bSorted);
-		return Bindings;
-	}
-
-	void SortBindings();
-
-	int32 GetRealBindingIndex(int32 InVirtualIndex) const
-	{
-		for (int32 Index = 0; Index < Bindings.Num(); ++Index)
-		{
-			if (Bindings[Index].VirtualIndex == InVirtualIndex)
-			{
-				return Index;
-			}
-		}
-
-		return -1;
-	}
-
-	uint32 InputAttachmentsMask = 0;
-
-private:
-	// Previous implementation supported bindings only for textures.
-	// However, layout(binding=%d) must be also used for uniform buffers.
-
-	EHlslShaderFrequency Stage;
-	TArray<FBinding> Bindings;
-
-	bool bSorted = false;
-};
-
 // InputAttachments
 // 0 - reserved for depth input, 1-8 for color
 extern const char* VULKAN_SUBPASS_FETCH[9];
diff --git a/Engine/Source/Developer/VulkanShaderFormat/Private/VulkanShaderCompiler.cpp b/Engine/Source/Developer/VulkanShaderFormat/Private/VulkanShaderCompiler.cpp
--- a/Engine/Source/Developer/VulkanShaderFormat/Private/VulkanShaderCompiler.cpp
+++ b/Engine/Source/Developer/VulkanShaderFormat/Private/VulkanShaderCompiler.cpp
@@ -165,19 +165,6 @@ uint32 ParseNumber(const T* Str, bool bEmptyIsZero = false)
 	return Num;
 }
 
-static bool ContainsBinding(const FVulkanBindingTable& BindingTable, const FString& Name)
-{
-	for (const FVulkanBindingTable::FBinding& Binding : BindingTable.GetBindings())
-	{
-		if (Binding.Name == Name)
-		{
-			return true;
-		}
-	}
-
-	return false;
-}
-
 static void GetResourceEntryFromUBMember(const FShaderResourceTableMap& ResourceTableMap, const FString& UBName, uint16 ResourceIndex, FUniformResourceEntry& OutEntry)
 {
 	for (const FUniformResourceEntry& Entry : ResourceTableMap.Resources)
@@ -313,23 +300,6 @@ struct FVulkanShaderCompilerInternalState
 	{
 		return (Version == EVulkanShaderVersion::ES3_1 || Version == EVulkanShaderVersion::ES3_1_ANDROID);
 	}
-	inline EHlslShaderFrequency GetHlslShaderFrequency() const
-	{
-		const EHlslShaderFrequency FrequencyTable[] =
-		{
-			HSF_VertexShader,
-			HSF_InvalidFrequency,
-			HSF_InvalidFrequency,
-			HSF_PixelShader,
-			(IsSM5() || IsSM6()) ? HSF_GeometryShader : HSF_InvalidFrequency,
-			HSF_ComputeShader,
-			(IsSM5() || IsSM6()) ? HSF_RayGen : HSF_InvalidFrequency,
-			(IsSM5() || IsSM6()) ? HSF_RayMiss : HSF_InvalidFrequency,
-			(IsSM5() || IsSM6()) ? HSF_RayHitGroup : HSF_InvalidFrequency,
-			(IsSM5() || IsSM6()) ? HSF_RayCallable : HSF_InvalidFrequency,
-		};
-		return FrequencyTable[Input.Target.Frequency];
-	}
 	inline FString GetDebugName() const
 	{
 		return Input.DumpDebugInfoPath.Right(Input.DumpDebugInfoPath.Len() - Input.DumpDebugInfoRootPath.Len());
@@ -356,6 +326,7 @@ struct FVulkanShaderCompilerInternalState
 	}
 };
 
+
 // Data structures that will get serialized into ShaderCompilerOutput
 struct VulkanShaderCompilerSerializedOutput
 {
@@ -368,12 +339,15 @@ struct VulkanShaderCompilerSerializedOutput
 	FShaderResourceTable ShaderResourceTable;
 	FVulkanSpirv Spirv;
 
+	uint32 InputAttachmentsMask = 0;
+
 	TSet<FString> UsedBindlessUB;
+
+	FShaderCompilerResourceTable CompilerSRT;
 };
 
 
-static int32 AddGlobal( const TArray<VkDescriptorType>& DescriptorTypes,
-						const FString& ParameterName,
+static int32 AddGlobal( const FString& ParameterName,
 						uint16 BindingIndex,
 						VulkanShaderCompilerSerializedOutput& SerializedOutput, 
 						const TArray<FString>& GlobalNames
@@ -406,14 +380,9 @@ static int32 AddGlobal( const TArray<VkDescriptorType>& DescriptorTypes,
 		}
 	}
 
-	const VkDescriptorType DescriptorType = DescriptorTypes[Entry->Binding];
-
-	GlobalInfo.OriginalBindingIndex = Entry->Binding;
 	SerializedOutput.Header.GlobalSpirvInfos[HeaderGlobalIndex] = FVulkanShaderHeader::FSpirvInfo(Entry->WordDescriptorSetIndex, Entry->WordBindingIndex);
 
-	const int32 GlobalDescriptorTypeIndex = SerializedOutput.Header.GlobalDescriptorTypes.Add(DescriptorTypeToBinding(DescriptorType));
-	GlobalInfo.TypeIndex = GlobalDescriptorTypeIndex;
-	GlobalInfo.CombinedSamplerStateAliasIndex = UINT16_MAX;
+	GlobalInfo.TypeIndex = Entry->DescriptorType;
 
 #if VULKAN_ENABLE_BINDING_DEBUG_NAMES
 	GlobalInfo.DebugName = ParameterName;
@@ -422,15 +391,15 @@ static int32 AddGlobal( const TArray<VkDescriptorType>& DescriptorTypes,
 	return HeaderGlobalIndex;
 }
 
-static void AddUBResources( const FString& UBName,
+static int32 AddUBResources(	const int32 HeaderUBIndex,
+							const FString& UBName,
 							const FShaderResourceTableMap& ResourceTableMap,
 							uint32 BufferIndex,
 							const TArray<uint32>& BindingArray,
-							const TArray<VkDescriptorType>& DescriptorTypes,
-							FVulkanShaderHeader::FUniformBufferInfo& OutUBInfo,
 							VulkanShaderCompilerSerializedOutput& SerializedOutput,
 							TArray<FString>& GlobalNames)
 {
+	int32 NumAdded = 0;
 	if (BindingArray.Num() > 0)
 	{
 		uint32 BufferOffset = BindingArray[BufferIndex];
@@ -453,27 +422,55 @@ static void AddUBResources( const FString& UBName,
 				FUniformResourceEntry ResourceTableEntry;
 				GetResourceEntryFromUBMember(ResourceTableMap, UBName, ResourceIndex, ResourceTableEntry);
 
-				FVulkanShaderHeader::FUBResourceInfo& UBResourceInfo = OutUBInfo.ResourceEntries.AddZeroed_GetRef();;
+				const int32 HeaderGlobalIndex = AddGlobal(ResourceTableEntry.UniformBufferMemberName, BindingIndex, SerializedOutput, GlobalNames);
+
+				// Fill the compiler SRT with info about this resource
+				{
+					SerializedOutput.CompilerSRT.ResourceTableBits |= 1 << HeaderUBIndex;
+					const uint32 ResourceMap = FRHIResourceTableEntry::Create(HeaderUBIndex, ResourceIndex, HeaderGlobalIndex);
+					switch ((EUniformBufferBaseType)ResourceTableEntry.Type)
+					{
+					case UBMT_TEXTURE:
+					case UBMT_RDG_TEXTURE:
+						SerializedOutput.CompilerSRT.TextureMap.Add(ResourceMap);
+						break;
+					case UBMT_SAMPLER:
+						SerializedOutput.CompilerSRT.SamplerMap.Add(ResourceMap);
+						break;
+					case UBMT_SRV:
+					case UBMT_RDG_TEXTURE_SRV:
+					case UBMT_RDG_TEXTURE_NON_PIXEL_SRV:
+					case UBMT_RDG_BUFFER_SRV:
+						SerializedOutput.CompilerSRT.ShaderResourceViewMap.Add(ResourceMap);
+						break;
+					case UBMT_RESOURCE_COLLECTION:
+						SerializedOutput.CompilerSRT.ResourceCollectionMap.Add(ResourceMap);
+						break;
+					case UBMT_UAV:
+					case UBMT_RDG_TEXTURE_UAV:
+					case UBMT_RDG_BUFFER_UAV:
+						SerializedOutput.CompilerSRT.UnorderedAccessViewMap.Add(ResourceMap);
+						break;
+					default:
+						check(false);
+					}
+
+					SerializedOutput.CompilerSRT.MaxBoundResourceTable = FMath::Max<uint32>(HeaderUBIndex+1, SerializedOutput.CompilerSRT.MaxBoundResourceTable);
+
+					NumAdded++;
+				}
 
-				const int32 HeaderGlobalIndex = AddGlobal(DescriptorTypes, ResourceTableEntry.UniformBufferMemberName, BindingIndex, SerializedOutput, GlobalNames);
-				UBResourceInfo.SourceUBResourceIndex = ResourceIndex;
-				UBResourceInfo.OriginalBindingIndex = BindingIndex;
-				UBResourceInfo.GlobalIndex = HeaderGlobalIndex;
-				UBResourceInfo.UBBaseType = (EUniformBufferBaseType)ResourceTableEntry.Type;
-#if VULKAN_ENABLE_BINDING_DEBUG_NAMES
-				UBResourceInfo.DebugName = ResourceTableEntry.UniformBufferMemberName;
-#endif
 				// Iterate to next info
 				ResourceInfo = *ResourceInfos++;
 			}
 			while (FRHIResourceTableEntry::GetUniformBufferIndex(ResourceInfo) == BufferIndex);
 		}
 	}
+	return NumAdded;
 }
 
 static void AddUniformBuffer(
 	const FShaderResourceTable& ShaderResourceTable,
-	const TArray<VkDescriptorType>& DescriptorTypes,
 	const FShaderCompilerInput& ShaderInput,
 	const FString& UBName,
 	uint16 BindingIndex,
@@ -500,6 +497,8 @@ static void AddUniformBuffer(
 	{
 		UBInfo.LayoutHash = 0;
 	}
+	const int32 NewSRTHashIndex = SerializedOutput.CompilerSRT.ResourceTableLayoutHashes.Emplace(UBInfo.LayoutHash);
+	check(HeaderUBIndex == NewSRTHashIndex);
 
 #if VULKAN_ENABLE_BINDING_DEBUG_NAMES
 	UBInfo.DebugName = UBName;
@@ -526,11 +525,15 @@ static void AddUniformBuffer(
 	// Add used resources...
 	if (ShaderResourceTable.ResourceTableBits & (1 << BindingIndex))
 	{
+		int32 NumAdded = 0;
+
 		// Make sure to process in the same order as when gathering names below
-		AddUBResources(UBName, ShaderInput.Environment.ResourceTableMap, BindingIndex, ShaderResourceTable.TextureMap, DescriptorTypes, UBInfo, SerializedOutput, GlobalNames);
-		AddUBResources(UBName, ShaderInput.Environment.ResourceTableMap, BindingIndex, ShaderResourceTable.SamplerMap, DescriptorTypes, UBInfo, SerializedOutput, GlobalNames);
-		AddUBResources(UBName, ShaderInput.Environment.ResourceTableMap, BindingIndex, ShaderResourceTable.ShaderResourceViewMap, DescriptorTypes, UBInfo, SerializedOutput, GlobalNames);
-		AddUBResources(UBName, ShaderInput.Environment.ResourceTableMap, BindingIndex, ShaderResourceTable.UnorderedAccessViewMap, DescriptorTypes, UBInfo, SerializedOutput, GlobalNames);
+		NumAdded += AddUBResources(HeaderUBIndex, UBName, ShaderInput.Environment.ResourceTableMap, BindingIndex, ShaderResourceTable.TextureMap, SerializedOutput, GlobalNames);
+		NumAdded += AddUBResources(HeaderUBIndex, UBName, ShaderInput.Environment.ResourceTableMap, BindingIndex, ShaderResourceTable.SamplerMap, SerializedOutput, GlobalNames);
+		NumAdded += AddUBResources(HeaderUBIndex, UBName, ShaderInput.Environment.ResourceTableMap, BindingIndex, ShaderResourceTable.ShaderResourceViewMap, SerializedOutput, GlobalNames);
+		NumAdded += AddUBResources(HeaderUBIndex, UBName, ShaderInput.Environment.ResourceTableMap, BindingIndex, ShaderResourceTable.UnorderedAccessViewMap, SerializedOutput, GlobalNames);
+
+		UBInfo.bHasResources = (NumAdded > 0) ? 1 : 0;
 	}
 	else
 	{
@@ -538,9 +541,6 @@ static void AddUniformBuffer(
 		checkf(!UBInfo.bOnlyHasResources, TEXT("UBName = %s, BindingIndex = %d"), *UBName, (int32)BindingIndex);
 	}
 
-	// Currently we don't support mismatched uniform buffer layouts/cbuffers with resources!
-	check(UniformBufferEntry || UBInfo.ResourceEntries.Num() == 0);
-
 	InOutParameterMap.RemoveParameterAllocation(*UBName);
 	InOutParameterMap.AddParameterAllocation(*UBName, HeaderUBIndex, (uint16)FVulkanShaderHeader::UniformBuffer, 1, EShaderParameterType::UniformBuffer);
 }
@@ -589,7 +589,7 @@ static void PrepareUBResourceEntryGlobals(const TArray<uint32>& BindingArray, co
 }
 
 
-static void PrepareGlobals(const FVulkanBindingTable& BindingTable, const FSpirvReflectBindings& SpirvReflectBindings, const FShaderResourceTable& SRT, const TMap<FString, FVulkanShaderHeader::EType>& EntryTypes, const FShaderCompilerInput& ShaderInput, const TArray<FString>& ParameterNames, FShaderParameterMap& ParameterMap, TArray<FString>& OutGlobalNames, FVulkanShaderHeader& OutHeader)
+static void PrepareGlobals(const VulkanShaderCompilerSerializedOutput& SerializedOutput, const FSpirvReflectBindings& SpirvReflectBindings, const FShaderResourceTable& SRT, const TMap<FString, FVulkanShaderHeader::EType>& EntryTypes, const FShaderCompilerInput& ShaderInput, const TArray<FString>& ParameterNames, FShaderParameterMap& ParameterMap, TArray<FString>& OutGlobalNames, FVulkanShaderHeader& OutHeader)
 {
 	auto IsSamplerState = [&SpirvReflectBindings](const FString& ParameterName)
 	{
@@ -697,9 +697,9 @@ static void PrepareGlobals(const FVulkanBindingTable& BindingTable, const FSpirv
 	}
 
 	// Now input attachments
-	if (BindingTable.InputAttachmentsMask != 0)
+	if (SerializedOutput.InputAttachmentsMask != 0)
 	{
-		uint32 InputAttachmentsMask = BindingTable.InputAttachmentsMask;
+		uint32 InputAttachmentsMask = SerializedOutput.InputAttachmentsMask;
 		for (int32 Index = 0; InputAttachmentsMask != 0; ++Index, InputAttachmentsMask>>= 1)
 		{
 			if (InputAttachmentsMask & 1)
@@ -712,8 +712,6 @@ static void PrepareGlobals(const FVulkanBindingTable& BindingTable, const FSpirv
 
 static void ConvertToHeader(
 	FShaderResourceTable& ShaderResourceTable,
-	const FVulkanBindingTable& BindingTable,
-	const TArray<VkDescriptorType>& DescriptorTypes,
 	const TMap<FString, FVulkanShaderHeader::EType>& EntryTypes,
 	const FShaderCompilerInput& ShaderInput,
 	const FSpirvReflectBindings& SpirvReflectBindings,
@@ -729,7 +727,7 @@ static void ConvertToHeader(
 	TArray<FString> ParameterNames;
 	InOutParameterMap.GetAllParameterNames(ParameterNames);
 
-	PrepareGlobals(BindingTable, SpirvReflectBindings, ShaderResourceTable, EntryTypes, ShaderInput, ParameterNames, InOutParameterMap, GlobalNames, OutHeader);
+	PrepareGlobals(SerializedOutput, SpirvReflectBindings, ShaderResourceTable, EntryTypes, ShaderInput, ParameterNames, InOutParameterMap, GlobalNames, OutHeader);
 
 	for (int32 ParameterIndex = 0; ParameterIndex < ParameterNames.Num(); ++ParameterIndex)
 	{
@@ -747,7 +745,7 @@ static void ConvertToHeader(
 			{
 			case FVulkanShaderHeader::Global:
 				{
-					const int32 HeaderGlobalIndex = AddGlobal(DescriptorTypes, ParameterName, BaseIndex, SerializedOutput, GlobalNames);
+					const int32 HeaderGlobalIndex = AddGlobal(ParameterName, BaseIndex, SerializedOutput, GlobalNames);
 
 					const FParameterAllocation* ParameterAllocation = InOutParameterMap.GetParameterMap().Find(*ParameterName);
 					check(ParameterAllocation);
@@ -760,7 +758,6 @@ static void ConvertToHeader(
 			case FVulkanShaderHeader::PackedGlobal:
 				{
 					FVulkanShaderHeader::FPackedGlobalInfo& PackedGlobalInfo = OutHeader.PackedGlobals.AddZeroed_GetRef();
-					PackedGlobalInfo.PackedTypeIndex = CrossCompiler::EPackedTypeIndex::HighP;
 					PackedGlobalInfo.PackedUBIndex = BufferIndex;
 					checkf(Size > 0, TEXT("Assertion failed for shader parameter: %s"), *ParameterName);
 					PackedGlobalInfo.ConstantDataSizeInFloats = Size / sizeof(float);
@@ -771,7 +768,7 @@ static void ConvertToHeader(
 				}
 				break;
 			case FVulkanShaderHeader::UniformBuffer:
-				AddUniformBuffer(ShaderResourceTable, DescriptorTypes, ShaderInput, ParameterName, BufferIndex, InOutParameterMap, SerializedOutput, GlobalNames);
+				AddUniformBuffer(ShaderResourceTable, ShaderInput, ParameterName, BufferIndex, InOutParameterMap, SerializedOutput, GlobalNames);
 				break;
 			default:
 				check(0);
@@ -781,12 +778,12 @@ static void ConvertToHeader(
 		else
 		{
 			// Not found means it's a new resource-only UniformBuffer
-			AddUniformBuffer(ShaderResourceTable, DescriptorTypes, ShaderInput, ParameterName, BufferIndex, InOutParameterMap, SerializedOutput, GlobalNames);
+			AddUniformBuffer(ShaderResourceTable, ShaderInput, ParameterName, BufferIndex, InOutParameterMap, SerializedOutput, GlobalNames);
 		}
 	}
 
 	// Finally check for subpass/input attachments
-	if (BindingTable.InputAttachmentsMask != 0)
+	if (SerializedOutput.InputAttachmentsMask != 0)
 	{
 		const static FVulkanShaderHeader::EAttachmentType AttachmentTypes[] = 
 		{
@@ -801,7 +798,7 @@ static void ConvertToHeader(
 			FVulkanShaderHeader::EAttachmentType::Color7
 		};
 
-		uint32 InputAttachmentsMask = BindingTable.InputAttachmentsMask;
+		uint32 InputAttachmentsMask = SerializedOutput.InputAttachmentsMask;
 		for (int32 Index = 0; InputAttachmentsMask != 0; ++Index, InputAttachmentsMask>>=1)
 		{
 			if ((InputAttachmentsMask & 1) == 0)
@@ -821,11 +818,8 @@ static void ConvertToHeader(
 				check(Entry);
 				check(Entry->Binding != -1);
 
-				GlobalInfo.OriginalBindingIndex = Entry->Binding;
 				OutHeader.GlobalSpirvInfos[HeaderGlobalIndex] = FVulkanShaderHeader::FSpirvInfo(Entry->WordDescriptorSetIndex, Entry->WordBindingIndex);
-				const int32 GlobalDescriptorTypeIndex = OutHeader.GlobalDescriptorTypes.Add(DescriptorTypeToBinding(DescriptorType));
-				GlobalInfo.TypeIndex = GlobalDescriptorTypeIndex;
-				GlobalInfo.CombinedSamplerStateAliasIndex = UINT16_MAX;
+				GlobalInfo.TypeIndex = (uint32)DescriptorType;
 #if VULKAN_ENABLE_BINDING_DEBUG_NAMES
 				GlobalInfo.DebugName = AttachmentName;
 #endif
@@ -838,72 +832,12 @@ static void ConvertToHeader(
 }
 
 
-// Fills the SRT using final values kept in the FVulkanShaderHeader.
-// NOTE: Uses GlobalIndex so it can be consumed directly at runtime.
-// NOTE: Keep in sync with BuildResourceTableMapping.
-static FShaderResourceTable BuildSRTFromHeader(const FVulkanShaderHeader& NEWHeader)
-{
-	FShaderCompilerResourceTable GenericSRT;
-
-	for (int32 UBIndex = 0; UBIndex < NEWHeader.UniformBuffers.Num(); ++UBIndex)
-	{
-		const FVulkanShaderHeader::FUniformBufferInfo& UBHeader = NEWHeader.UniformBuffers[UBIndex];
-
-		GenericSRT.ResourceTableLayoutHashes.Emplace(UBHeader.LayoutHash);
-		if (UBHeader.ResourceEntries.Num() > 0)
-		{
-			GenericSRT.ResourceTableBits |= 1 << UBIndex;
-
-			for (const FVulkanShaderHeader::FUBResourceInfo& UBRes : UBHeader.ResourceEntries)
-			{
-				uint32 ResourceMap = FRHIResourceTableEntry::Create(UBIndex, UBRes.SourceUBResourceIndex, UBRes.GlobalIndex);
-				switch (UBRes.UBBaseType)
-				{
-				case UBMT_TEXTURE:
-				case UBMT_RDG_TEXTURE:
-					GenericSRT.TextureMap.Add(ResourceMap);
-					break;
-				case UBMT_SAMPLER:
-					GenericSRT.SamplerMap.Add(ResourceMap);
-					break;
-				case UBMT_SRV:
-				case UBMT_RDG_TEXTURE_SRV:
-				case UBMT_RDG_TEXTURE_NON_PIXEL_SRV:
-				case UBMT_RDG_BUFFER_SRV:
-					GenericSRT.ShaderResourceViewMap.Add(ResourceMap);
-					break;
-				case UBMT_RESOURCE_COLLECTION:
-					GenericSRT.ResourceCollectionMap.Add(ResourceMap);
-					break;
-				case UBMT_UAV:
-				case UBMT_RDG_TEXTURE_UAV:
-				case UBMT_RDG_BUFFER_UAV:
-					GenericSRT.UnorderedAccessViewMap.Add(ResourceMap);
-					break;
-				default:
-					check(false);
-				}
-			}
-		}
-	}
-
-	GenericSRT.MaxBoundResourceTable = NEWHeader.UniformBuffers.Num();
-
-	FShaderResourceTable ShaderResourceTable;
-	UE::ShaderCompilerCommon::BuildShaderResourceTable(GenericSRT, ShaderResourceTable);
-
-	return ShaderResourceTable;
-}
-
-
-
 static void BuildShaderOutput(
 	VulkanShaderCompilerSerializedOutput& SerializedOutput,
 	FShaderCompilerOutput&		ShaderOutput,
 	const FVulkanShaderCompilerInternalState& InternalState,
 	const FSpirvReflectBindings& SpirvReflectBindings,
 	TMap<FString, FVulkanShaderHeader::EType>& EntryTypes,
-	const FVulkanBindingTable&	BindingTable,
 	const FString&				DebugName,
 	uint32						PackedGlobalArraySize,
 	TBitArray<>&				UsedUniformBufferSlots
@@ -1002,16 +936,6 @@ static void BuildShaderOutput(
 	
 	const int32 StageOffset = InternalState.bSupportsBindless ? (ShaderStage::GetStageForFrequency(Frequency) * VulkanBindless::MaxUniformBuffersPerStage) : 0;
 
-	TArray<VkDescriptorType> DescriptorTypes;
-	const TArray<FVulkanBindingTable::FBinding>& HlslccBindings = BindingTable.GetBindings();
-	for (int32 Index = 0; Index < HlslccBindings.Num(); ++Index)
-	{
-		const FVulkanBindingTable::FBinding& Binding = HlslccBindings[Index];
-		DescriptorTypes.Add(BindingToDescriptorType(Binding.Type));
-	}
-
-	//#todo-rco: When using regular UBs, also set UsedUniformBufferSlots[] = 1
-
 	TArray<FString> OriginalParameters;
 	ShaderOutput.ParameterMap.GetAllParameterNames(OriginalParameters);
 
@@ -1049,8 +973,6 @@ static void BuildShaderOutput(
 	if (PackedGlobalArraySize > 0)
 	{
 		FVulkanShaderHeader::FPackedUBInfo& PackedUB = NEWHeader.PackedUBs.AddZeroed_GetRef();
-		PackedUB.OriginalBindingIndex = StageOffset;
-		PackedUB.PackedTypeIndex = CrossCompiler::EPackedTypeIndex::HighP;
 		PackedUB.SizeInBytes = Align(PackedGlobalArraySize, 16u);
 
 		const FVulkanSpirv::FEntry* Entry = SerializedOutput.Spirv.GetEntryByBindingIndex(StageOffset);
@@ -1059,16 +981,15 @@ static void BuildShaderOutput(
 		PackedUB.SPIRVBindingIndexOffset = Entry->WordBindingIndex;
 	}
 
-	ConvertToHeader(ShaderResourceTable, BindingTable, DescriptorTypes, EntryTypes, ShaderInput, SpirvReflectBindings, ShaderOutput.ParameterMap, SerializedOutput);
-	check(NEWHeader.EmulatedUBsCopyInfo.Num() == 0);
+	ConvertToHeader(ShaderResourceTable, EntryTypes, ShaderInput, SpirvReflectBindings, ShaderOutput.ParameterMap, SerializedOutput);
 
 	if (ShaderInput.Environment.CompilerFlags.Contains(CFLAG_ExtraShaderData))
 	{
 		NEWHeader.DebugName = ShaderInput.GenerateShaderName();
 	}
 
-	// Build the SRT for this shader from the NEWHeader
-	SerializedOutput.ShaderResourceTable = BuildSRTFromHeader(NEWHeader);
+	// Build the SRT for this shader from accumulated data in the Compiler SRT
+	UE::ShaderCompilerCommon::BuildShaderResourceTable(SerializedOutput.CompilerSRT, SerializedOutput.ShaderResourceTable);
 
 	ShaderOutput.bSucceeded = true;
 
@@ -1255,8 +1176,7 @@ static bool BuildShaderOutputFromSpirv(
 	CrossCompiler::FShaderConductorContext&	CompilerContext,
 	const FVulkanShaderCompilerInternalState& InternalState,
 	VulkanShaderCompilerSerializedOutput&   SerializedOutput,
-	FShaderCompilerOutput&					Output,
-	FVulkanBindingTable&					BindingTable
+	FShaderCompilerOutput&					Output
 )
 {
 	// Reflect SPIR-V module with SPIRV-Reflect library
@@ -1276,52 +1196,10 @@ static bool BuildShaderOutputFromSpirv(
 	FSpirvReflectBindings Bindings;
 	GatherSpirvReflectionBindings(Reflection, Bindings, SerializedOutput.UsedBindlessUB, InternalState);
 
-	// Build binding table
-	TMap<const SpvReflectDescriptorBinding*, int32> BindingToIndexMap;
-
 	const FString UBOGlobalsNameSpv(ANSI_TO_TCHAR(CrossCompiler::FShaderConductorContext::GetIdentifierTable().GlobalsUniformBuffer));
 	const FString UBORootParamNameSpv(FShaderParametersMetadata::kRootUniformBufferBindingName);
 
-	auto RegisterBindings = [&BindingTable, &BindingToIndexMap, &UBOGlobalsNameSpv, &UBORootParamNameSpv]
-							(TArray<SpvReflectDescriptorBinding*>& Bindings, const char* BlockName, EVulkanBindingType::EType BindingType)
-	{
-		for (const SpvReflectDescriptorBinding* Binding : Bindings)
-		{
-			const bool bIsGlobalOrRootBuffer = ((UBOGlobalsNameSpv == Binding->name) || (UBORootParamNameSpv == Binding->name));
-			if (((BindingType == EVulkanBindingType::PackedUniformBuffer) && !bIsGlobalOrRootBuffer) ||
-				((BindingType == EVulkanBindingType::UniformBuffer) && bIsGlobalOrRootBuffer))
-			{
-				continue;
-			}
-
-			const int32 BindingIndex = BindingTable.RegisterBinding(Binding->name, BlockName, BindingType);
-			BindingToIndexMap.Add(Binding, BindingIndex);
-
-			if (BindingType == EVulkanBindingType::InputAttachment)
-			{
-				BindingTable.InputAttachmentsMask |= (1u << Binding->input_attachment_index);
-			}
-		}
-	};
-
-	RegisterBindings(Bindings.UniformBuffers, "h", EVulkanBindingType::PackedUniformBuffer);
-	RegisterBindings(Bindings.UniformBuffers, "u", EVulkanBindingType::UniformBuffer);
-	RegisterBindings(Bindings.InputAttachments, "a", EVulkanBindingType::InputAttachment);
-
-	RegisterBindings(Bindings.TBufferUAVs, "u", EVulkanBindingType::StorageTexelBuffer);
-	RegisterBindings(Bindings.SBufferUAVs, "u", EVulkanBindingType::StorageBuffer);
-	RegisterBindings(Bindings.TextureUAVs, "u", EVulkanBindingType::StorageImage);
-
-	RegisterBindings(Bindings.TBufferSRVs, "s", EVulkanBindingType::UniformTexelBuffer);
-	checkf(Bindings.SBufferSRVs.IsEmpty(), TEXT("GatherSpirvReflectionBindings should have dumped all SBufferSRVs into SBufferUAVs."));
-	RegisterBindings(Bindings.TextureSRVs, "s", EVulkanBindingType::Image);
-
-	RegisterBindings(Bindings.Samplers, "z", EVulkanBindingType::Sampler);
-	RegisterBindings(Bindings.AccelerationStructures, "r", EVulkanBindingType::AccelerationStructure);
-
-	// Sort binding table
-	BindingTable.SortBindings();
-
+	uint32& InputAttachmentsMask = SerializedOutput.InputAttachmentsMask;
 	uint32 PackedGlobalArraySize = 0;
 	TBitArray<> UsedUniformBufferSlots;
 	const int32 MaxNumBits = VulkanBindless::MaxUniformBuffersPerStage * SF_NumFrequencies;
@@ -1331,31 +1209,33 @@ static bool BuildShaderOutputFromSpirv(
 
 	// Final descriptor binding numbers for all other resource types
 	{
-		const int32 StageOffset = InternalState.bSupportsBindless ? (ShaderStage::GetStageForFrequency(InternalState.GetShaderFrequency()) * VulkanBindless::MaxUniformBuffersPerStage) : 0;
-		const uint32_t DescSetNumber = InternalState.bSupportsBindless ? (uint32_t)VulkanBindless::BindlessSingleUseUniformBufferSet : (uint32_t)SPV_REFLECT_SET_NUMBER_DONT_CHANGE;
+		const ShaderStage::EStage UEStage = ShaderStage::GetStageForFrequency(InternalState.GetShaderFrequency());
+		const int32 StageOffset = InternalState.bSupportsBindless ? (UEStage * VulkanBindless::MaxUniformBuffersPerStage) : 0;
+		const uint32_t DescSetNumber = InternalState.bSupportsBindless ? (uint32_t)VulkanBindless::BindlessSingleUseUniformBufferSet : (uint32_t)UEStage;
 
-		auto AddReflectionInfos = [&](TArray<SpvReflectDescriptorBinding*>& BindingArray, EVulkanBindingType::EType BindingType, int32 BindingOffset)
+		int32 GlobalBindingIndex = StageOffset;
+
+		auto AddReflectionInfos = [&](TArray<SpvReflectDescriptorBinding*>& BindingArray, const VkDescriptorType DescriptorType, int32 BindingOffset, bool bIsPackedUniformBuffer=false)
 		{
 			for (const SpvReflectDescriptorBinding* Binding : BindingArray)
 			{
-				checkf(!InternalState.bSupportsBindless || (BindingType == EVulkanBindingType::UniformBuffer) || (BindingType == EVulkanBindingType::PackedUniformBuffer),
+				checkf(!InternalState.bSupportsBindless || (DescriptorType == VK_DESCRIPTOR_TYPE_UNIFORM_BUFFER),
 					TEXT("Bindless shaders should only have uniform buffers."));
 
 				const FString ResourceName(ANSI_TO_TCHAR(Binding->name));
 
 				const bool bIsGlobalOrRootBuffer = ((UBOGlobalsNameSpv == ResourceName) || (UBORootParamNameSpv == ResourceName));
-				if (((BindingType == EVulkanBindingType::PackedUniformBuffer) && !bIsGlobalOrRootBuffer) ||
-					((BindingType == EVulkanBindingType::UniformBuffer) && bIsGlobalOrRootBuffer))
+				if ((bIsPackedUniformBuffer && !bIsGlobalOrRootBuffer) || ((!bIsPackedUniformBuffer) && bIsGlobalOrRootBuffer))
 				{
 					continue;
 				}
 
-				const int32 BindingIndex = BindingTable.GetRealBindingIndex(BindingToIndexMap[Binding]) + StageOffset;
+				const int32 BindingIndex = GlobalBindingIndex++;
 
 				const SpvReflectResult SpvResult = Reflection.ChangeDescriptorBindingNumbers(Binding, BindingIndex, DescSetNumber);
 				check(SpvResult == SPV_REFLECT_RESULT_SUCCESS);
 
-				const int32 ReflectionSlot = SerializedOutput.Spirv.ReflectionInfo.Add(FVulkanSpirv::FEntry(ResourceName, BindingIndex));
+				const int32 ReflectionSlot = SerializedOutput.Spirv.ReflectionInfo.Add(FVulkanSpirv::FEntry(ResourceName, BindingIndex, (uint32)DescriptorType));
 				check(InternalState.ParameterParser);
 				const FShaderParameterParser::FParsedShaderParameter* ParsedParam = InternalState.ParameterParser->FindParameterInfosUnsafe(ResourceName);
 
@@ -1373,18 +1253,18 @@ static bool BuildShaderOutputFromSpirv(
 					}*/
 				};
 
-				switch (BindingType)
+				switch (DescriptorType)
 				{
-				case EVulkanBindingType::StorageTexelBuffer:
-				case EVulkanBindingType::StorageBuffer:
-				case EVulkanBindingType::StorageImage:
+				case VK_DESCRIPTOR_TYPE_STORAGE_TEXEL_BUFFER:
+				case VK_DESCRIPTOR_TYPE_STORAGE_BUFFER:
+				case VK_DESCRIPTOR_TYPE_STORAGE_IMAGE:
 					HandleReflectedShaderUAV(ResourceName, BindingOffset, ReflectionSlot, 1, Output);
 					EntryTypes.Add(ResourceName, FVulkanShaderHeader::Global);
 
 					AddShaderValidationType(BindingOffset, ParsedParam, Output);
 					break;
 
-				case EVulkanBindingType::Image:
+				case VK_DESCRIPTOR_TYPE_SAMPLED_IMAGE:
 					// todo-jn: Could verify that we have the samplers...
 					//for (uint32 UsageIndex = 0; UsageIndex < Binding->usage_binding_count; ++UsageIndex)
 					//{
@@ -1393,31 +1273,32 @@ static bool BuildShaderOutputFromSpirv(
 					//}
 					[[fallthrough]];
 
-				case EVulkanBindingType::UniformTexelBuffer:
+				case VK_DESCRIPTOR_TYPE_UNIFORM_TEXEL_BUFFER:
 					HandleReflectedShaderResource(ResourceName, BindingOffset, ReflectionSlot, 1, Output);
 					EntryTypes.Add(ResourceName, FVulkanShaderHeader::Global);
 
 					AddShaderValidationType(BindingOffset, ParsedParam, Output);
 					break;
 
-				case EVulkanBindingType::Sampler:
+				case VK_DESCRIPTOR_TYPE_SAMPLER:
 					HandleReflectedShaderSampler(ResourceName, ReflectionSlot, Output);
 					//HandleReflectedShaderSampler(ResourceName, BindingOffset, ReflectionSlot, 1, Output);
 					EntryTypes.Add(ResourceName, FVulkanShaderHeader::Global);
 					break;
 
-				case EVulkanBindingType::AccelerationStructure:
+				case VK_DESCRIPTOR_TYPE_ACCELERATION_STRUCTURE_KHR:
 					HandleReflectedShaderResource(ResourceName, BindingOffset, ReflectionSlot, 1, Output);
 					EntryTypes.Add(ResourceName, FVulkanShaderHeader::Global);
 
 					AddShaderValidationType(BindingOffset, ParsedParam, Output);
 					break;
 
-				case EVulkanBindingType::InputAttachment:
-					// Do Nothing
+				case VK_DESCRIPTOR_TYPE_INPUT_ATTACHMENT:
+					InputAttachmentsMask |= (1u << Binding->input_attachment_index);
 					break;
 
-				case EVulkanBindingType::PackedUniformBuffer:
+				case VK_DESCRIPTOR_TYPE_UNIFORM_BUFFER:
+					if (bIsPackedUniformBuffer)
 					{
 						// Use the given global ResourceName instead of patching it to _Globals_h
 						check(!UsedUniformBufferSlots[ReflectionSlot]);
@@ -1443,6 +1324,7 @@ static bool BuildShaderOutputFromSpirv(
 							// Add all members of global ub, and only bindless samplers/resources for root param
 							if (!InternalState.UseRootParametersStructure() || BindlessParameterType != EShaderParameterType::LooseData)
 							{
+								check(BindingOffset == 0); // Global constants should always be the first UB
 								HandleReflectedGlobalConstantBufferMember(
 									MemberName,
 									BindingOffset,
@@ -1457,9 +1339,7 @@ static bool BuildShaderOutputFromSpirv(
 							PackedGlobalArraySize = FMath::Max<uint32>((Member.absolute_offset + Member.size), PackedGlobalArraySize);
 						}
 					}
-					break;
-
-				case EVulkanBindingType::UniformBuffer:
+					else
 					{
 						check(!UsedUniformBufferSlots[ReflectionSlot]);
 						UsedUniformBufferSlots[ReflectionSlot] = true;
@@ -1481,23 +1361,23 @@ static bool BuildShaderOutputFromSpirv(
 		};
 
 		// Process Globals first (PackedUniformBuffer) and then regular UBs
-		const int32 GlobalUBCount = AddReflectionInfos(Bindings.UniformBuffers, EVulkanBindingType::PackedUniformBuffer, 0);
-		int32 UBOBindings = AddReflectionInfos(Bindings.UniformBuffers, EVulkanBindingType::UniformBuffer, 0) + GlobalUBCount;
+		const int32 GlobalUBCount = AddReflectionInfos(Bindings.UniformBuffers, VK_DESCRIPTOR_TYPE_UNIFORM_BUFFER, 0, true);
+		int32 UBOBindings = AddReflectionInfos(Bindings.UniformBuffers, VK_DESCRIPTOR_TYPE_UNIFORM_BUFFER, 0) + GlobalUBCount;
 
-		AddReflectionInfos(Bindings.InputAttachments, EVulkanBindingType::InputAttachment, 0);
+		AddReflectionInfos(Bindings.InputAttachments, VK_DESCRIPTOR_TYPE_INPUT_ATTACHMENT, 0);
 
 		int32 UAVBindings = 0;
-		UAVBindings = AddReflectionInfos(Bindings.TBufferUAVs, EVulkanBindingType::StorageTexelBuffer, UAVBindings);
-		UAVBindings = AddReflectionInfos(Bindings.SBufferUAVs, EVulkanBindingType::StorageBuffer, UAVBindings);
-		UAVBindings = AddReflectionInfos(Bindings.TextureUAVs, EVulkanBindingType::StorageImage, UAVBindings);
+		UAVBindings = AddReflectionInfos(Bindings.TBufferUAVs, VK_DESCRIPTOR_TYPE_STORAGE_TEXEL_BUFFER, UAVBindings);
+		UAVBindings = AddReflectionInfos(Bindings.SBufferUAVs, VK_DESCRIPTOR_TYPE_STORAGE_BUFFER, UAVBindings);
+		UAVBindings = AddReflectionInfos(Bindings.TextureUAVs, VK_DESCRIPTOR_TYPE_STORAGE_IMAGE, UAVBindings);
 
 		int32 SRVBindings = 0;
-		SRVBindings = AddReflectionInfos(Bindings.TBufferSRVs, EVulkanBindingType::UniformTexelBuffer, SRVBindings);
+		SRVBindings = AddReflectionInfos(Bindings.TBufferSRVs, VK_DESCRIPTOR_TYPE_UNIFORM_TEXEL_BUFFER, SRVBindings);
 		checkf(Bindings.SBufferSRVs.IsEmpty(), TEXT("GatherSpirvReflectionBindings should have dumped all SBufferSRVs into SBufferUAVs."));
-		SRVBindings = AddReflectionInfos(Bindings.TextureSRVs, EVulkanBindingType::Image, SRVBindings);
+		SRVBindings = AddReflectionInfos(Bindings.TextureSRVs, VK_DESCRIPTOR_TYPE_SAMPLED_IMAGE, SRVBindings);
 
-		Output.NumTextureSamplers = AddReflectionInfos(Bindings.Samplers, EVulkanBindingType::Sampler, 0);
-		AddReflectionInfos(Bindings.AccelerationStructures, EVulkanBindingType::AccelerationStructure, 0);
+		Output.NumTextureSamplers = AddReflectionInfos(Bindings.Samplers, VK_DESCRIPTOR_TYPE_SAMPLER, 0);
+		AddReflectionInfos(Bindings.AccelerationStructures, VK_DESCRIPTOR_TYPE_ACCELERATION_STRUCTURE_KHR, 0);
 	}
 
 	Output.Target = InternalState.Input.Target;
@@ -1547,7 +1427,6 @@ static bool BuildShaderOutputFromSpirv(
 		InternalState,
 		Bindings,
 		EntryTypes,
-		BindingTable,
 		InternalState.GetDebugName(),
 		PackedGlobalArraySize,
 		UsedUniformBufferSlots
@@ -1749,7 +1628,7 @@ static void VulkanCreateDXCCompileBatchFiles(
 }
 
 // Quick and dirty way to get the location of the entrypoint in the source
-// NOTE: Preprocessed shaders have mcros resolves and comments removed, it makes this easier...
+// NOTE: Preprocessed shaders have macros resolved and comments removed, it makes this easier...
 static FString ParseEntrypointDecl(FShaderSource::FViewType PreprocessedShader, FStringView Entrypoint)
 {
 	FShaderSource::FStringType EntrypointConverted(Entrypoint);
@@ -1900,7 +1779,6 @@ static bool CompileWithShaderConductor(
 {
 	const FShaderCompilerInput& Input = InternalState.Input;
 
-	FVulkanBindingTable BindingTable(InternalState.GetHlslShaderFrequency());
 	CrossCompiler::FShaderConductorContext CompilerContext;
 
 	// Inject additional macro definitions to circumvent missing features: external textures
@@ -1963,7 +1841,7 @@ static bool CompileWithShaderConductor(
 	Patch64bitSamplers(SerializedOutput.Spirv);
 
 	// Build shader output and binding table
-	Output.bSucceeded = BuildShaderOutputFromSpirv(CompilerContext, InternalState, SerializedOutput, Output, BindingTable);
+	Output.bSucceeded = BuildShaderOutputFromSpirv(CompilerContext, InternalState, SerializedOutput, Output);
 
 	// Flush warnings
 	CompilerContext.FlushErrors(Output.Errors);
@@ -1973,56 +1851,6 @@ static bool CompileWithShaderConductor(
 #endif // PLATFORM_MAC || PLATFORM_WINDOWS || PLATFORM_LINUX
 
 
-
-static void RemoveUnusedBindlessHeaps(FString& PreprocessedShaderSource, const TCHAR* HeapType)
-{
-	const FString HeapIdentifier = TEXT("VULKAN_") + FString(HeapType) + TEXT("_HEAP(");
-
-	int32 SearchIndex = PreprocessedShaderSource.Find(HeapIdentifier, ESearchCase::CaseSensitive, ESearchDir::FromStart, -1);
-	while (SearchIndex != INDEX_NONE)
-	{
-		const int32 TypeNameStartIndex = SearchIndex + HeapIdentifier.Len();
-		const int32 TypeNameEndIndex = PreprocessedShaderSource.Find(TEXT(")"), ESearchCase::CaseSensitive, ESearchDir::FromStart, TypeNameStartIndex);
-		FString TypeName(TypeNameEndIndex - TypeNameStartIndex, &PreprocessedShaderSource[TypeNameStartIndex]);
-		TypeName.TrimStartAndEndInline();
-
-		// Make sure it's one of our automatically generated types
-		if (TypeName.StartsWith(TEXT("SafeType")))
-		{
-			// Ugly and fast way to make sure ity's a heap declaration (this should catch the generated ones at least)
-			if ((PreprocessedShaderSource[TypeNameEndIndex + 1] == '[') &&
-				(PreprocessedShaderSource[TypeNameEndIndex + 2] == ']') &&
-				(PreprocessedShaderSource[TypeNameEndIndex + 3] == ';'))
-			{
-				int32 FirstUseIndex = PreprocessedShaderSource.Find(TypeName, ESearchCase::CaseSensitive, ESearchDir::FromStart, TypeNameEndIndex + 4);
-				while (FirstUseIndex != INDEX_NONE)
-				{
-					const int32 NextChar = FirstUseIndex + TypeName.Len();
-					if (FChar::IsWhitespace(PreprocessedShaderSource[NextChar]) || PreprocessedShaderSource[NextChar] == ')')
-					{
-						break;
-					}
-					FirstUseIndex = PreprocessedShaderSource.Find(TypeName, ESearchCase::CaseSensitive, ESearchDir::FromStart, NextChar);
-				}
-
-				if (FirstUseIndex == INDEX_NONE)
-				{
-					const int32 DeclarationBeginIndex = PreprocessedShaderSource.Find(TypeName, ESearchCase::CaseSensitive, ESearchDir::FromEnd, SearchIndex);
-					if ((DeclarationBeginIndex != INDEX_NONE) && ((SearchIndex - DeclarationBeginIndex) < (TypeName.Len() + 4)))
-					{
-						for (int32 Idx = DeclarationBeginIndex; Idx <= (TypeNameEndIndex + 3); Idx++)
-						{
-							PreprocessedShaderSource[Idx] = ' ';
-						}
-					}
-				}
-			}
-		}
-
-		SearchIndex = PreprocessedShaderSource.Find(HeapIdentifier, ESearchCase::CaseSensitive, ESearchDir::FromStart, TypeNameEndIndex + 4);
-	}
-}
-
 void ModifyVulkanCompilerInput(FShaderCompilerInput& Input)
 {
 	FVulkanShaderCompilerInternalState InternalState(Input, nullptr);
@@ -2418,33 +2246,6 @@ void CompileVulkanShader(const FShaderCompilerInput& Input, const FShaderPreproc
 
 	FVulkanShaderCompilerInternalState InternalState(Input, &ShaderParameterParser);
 
-	//TODO: this additional step causes problems for the error remapping that occurs when the preprocessed job cache is enabled
-	// (the additional deadstripping step causes further changes to line numbers, removed blocks, etc).
-	//if (InternalState.bSupportsBindless)
-	//{
-	//	// Clean up the code a bit, it's unreadable otherwise with all the unused heaps left around
-	//	// Re-run the dead stripper after removing these unused heaps
-	//	RemoveUnusedBindlessHeaps(PreprocessedSource, TEXT("SAMPLER"));
-	//	RemoveUnusedBindlessHeaps(PreprocessedSource, TEXT("RESOURCE"));
-	//	UE::ShaderCompilerCommon::RemoveDeadCode(PreprocessedSource, Input.EntryPointName, Input.RequiredSymbols, Output.Errors);
-
-	//	if (InternalState.bDebugDump)
-	//	{
-	//		DumpDebugShaderText(Input, PreprocessedSource, TEXT("bindless.final.hlsl"));
-	//	}
-	//}
-
-	const EHlslShaderFrequency HlslFrequency = InternalState.GetHlslShaderFrequency();
-	if (HlslFrequency == HSF_InvalidFrequency)
-	{
-		Output.bSucceeded = false;
-		FShaderCompilerError& NewError = Output.Errors.AddDefaulted_GetRef();
-		NewError.StrippedErrorMessage = FString::Printf(
-			TEXT("%s shaders not supported for use in Vulkan."),
-			CrossCompiler::GetFrequencyName(InternalState.GetShaderFrequency()));
-		return;
-	}
-
 	if (InternalState.bUseBindlessUniformBuffer)
 	{
 		InternalState.AllBindlessUBs = ConvertUBToBindless(PreprocessedSource);
@@ -2498,6 +2299,15 @@ void CompileVulkanShader(const FShaderCompilerInput& Input, const FShaderPreproc
 	}
 #endif // PLATFORM_MAC || PLATFORM_WINDOWS || PLATFORM_LINUX
 	
+	if (InternalState.bUseBindlessUniformBuffer)
+	{
+		// HACK: Because of heavy code alterations with bindless ray tracing shaders, line numbers will be all over the place.  Remove the tag that leads to remapping...
+		for (FShaderCompilerError& ErrorMsg : Output.Errors)
+		{
+			ErrorMsg.StrippedErrorMessage.ReplaceInline(TEXT("__UE_FILENAME_SENTINEL"), *Input.GenerateShaderName());
+		}
+	}
+
 	if (Input.Environment.CompilerFlags.Contains(CFLAG_ExtraShaderData))
 	{
 		Output.ShaderCode.AddOptionalData(FShaderCodeName::Key, TCHAR_TO_UTF8(*Input.GenerateShaderName()));
diff --git a/Engine/Source/Developer/VulkanShaderFormat/Private/VulkanShaderFormat.h b/Engine/Source/Developer/VulkanShaderFormat/Private/VulkanShaderFormat.h
--- a/Engine/Source/Developer/VulkanShaderFormat/Private/VulkanShaderFormat.h
+++ b/Engine/Source/Developer/VulkanShaderFormat/Private/VulkanShaderFormat.h
@@ -4,7 +4,6 @@
 
 #include "CoreMinimal.h"
 #include "SpirvCommon.h"
-#include "hlslcc.h"
 
 
 struct FVulkanSpirv : FSpirv
@@ -12,15 +11,17 @@ struct FVulkanSpirv : FSpirv
 	struct FEntry
 	{
 		FEntry() = default;
-		FEntry(const FString& InName, int32 InBinding)
+		FEntry(const FString& InName, int32 InBinding, uint32 InDescriptorType)
 			: Name(InName)
 			, Binding(InBinding)
+			, DescriptorType(InDescriptorType)
 		{
 		}
 		FString Name;
 		int32 Binding = -1;
 
 		uint32 DescriptorSet = UINT32_MAX;
+		uint32 DescriptorType = UINT32_MAX;
 
 		// Index into the Spirv Word containing the descriptor set decoration
 		uint32 WordDescriptorSetIndex = UINT32_MAX;
diff --git a/Engine/Source/Runtime/VulkanRHI/Private/VulkanCommands.cpp b/Engine/Source/Runtime/VulkanRHI/Private/VulkanCommands.cpp
--- a/Engine/Source/Runtime/VulkanRHI/Private/VulkanCommands.cpp
+++ b/Engine/Source/Runtime/VulkanRHI/Private/VulkanCommands.cpp
@@ -540,7 +540,7 @@ void FVulkanCommandListContext::RHISetShaderUniformBuffer(FRHIGraphicsShader* Sh
 		}
 	}
 
-	if (HeaderUBInfo.ResourceEntries.Num())
+	if (HeaderUBInfo.bHasResources)
 	{
 		checkSlow(Shader->Frequency < SF_NumStandardFrequencies);
 		check(BufferIndex < MAX_UNIFORM_BUFFERS_PER_SHADER_STAGE);
@@ -598,7 +598,7 @@ void FVulkanCommandListContext::RHISetShaderUniformBuffer(FRHIComputeShader* Com
 		}
 	}
 
-	if (HeaderUBInfo.ResourceEntries.Num())
+	if (HeaderUBInfo.bHasResources)
 	{
 		checkSlow(ComputeShaderRHI->GetFrequency() == SF_Compute);
 		check(BufferIndex < MAX_UNIFORM_BUFFERS_PER_SHADER_STAGE);
diff --git a/Engine/Source/Runtime/VulkanRHI/Private/VulkanDescriptorSets.h b/Engine/Source/Runtime/VulkanRHI/Private/VulkanDescriptorSets.h
--- a/Engine/Source/Runtime/VulkanRHI/Private/VulkanDescriptorSets.h
+++ b/Engine/Source/Runtime/VulkanRHI/Private/VulkanDescriptorSets.h
@@ -58,7 +58,6 @@ struct FDescriptorSetRemappingInfo
 
 		bool					bHasConstantData = false;
 		const bool				bPadding = false;	// padding is need on memcmp/MemCrc to make sure mem align
-		//bool					bIsRedundant = false;
 	};
 
 	struct FSetInfo
@@ -173,12 +172,9 @@ struct FDescriptorSetRemappingInfo
 		return false;
 	}
 
-	uint32 AddGlobal(uint32 Stage, int32 GlobalIndex, uint32 NewDescriptorSet, VkDescriptorType InType, uint16 CombinedSamplerStateAlias)
+	uint32 AddGlobal(uint32 Stage, int32 GlobalIndex, uint32 NewDescriptorSet, VkDescriptorType InType)
 	{
-		// Combined Image Samplers point both the texture and the sampler to the same descriptor
-		check(CombinedSamplerStateAlias == UINT16_MAX || InType == VK_DESCRIPTOR_TYPE_COMBINED_IMAGE_SAMPLER);
-
-		uint32 NewBindingIndex = CombinedSamplerStateAlias == UINT16_MAX ? SetInfos[NewDescriptorSet].Types.Add(InType) : StageInfos[Stage].Globals[CombinedSamplerStateAlias].NewBindingIndex;
+		const uint32 NewBindingIndex = SetInfos[NewDescriptorSet].Types.Add(InType);
 
 		int32 RemappingIndex = StageInfos[Stage].Globals.AddDefaulted();
 		check(RemappingIndex == GlobalIndex);
@@ -241,24 +237,6 @@ struct FDescriptorSetRemappingInfo
 		UBRemapping.bHasConstantData = false;
 	}
 
-/*
-	uint32 AddUBResourceEntryToSetAndRemapping(int32 Stage, int32 UniformBufferIndex, int32 ResourceEntryIndex, uint32 NewDescriptorSet, VkDescriptorType InType, uint8 CombinedSamplerStateAlias)
-	{
-		// Combined Image Samplers point both the texture and the sampler to the same descriptor
-		check(CombinedSamplerStateAlias == UINT8_MAX || InType == VK_DESCRIPTOR_TYPE_COMBINED_IMAGE_SAMPLER);
-		uint32 NewBindingIndex = CombinedSamplerStateAlias == UINT8_MAX ? SetInfos[NewDescriptorSet].Types.Add(InType) : StageInfos[Stage].UniformBuffers[UniformBufferIndex].EntriesRemappingInfo[CombinedSamplerStateAlias].NewBindingIndex;
-
-		int32 UBEntryRemappingIndex = StageInfos[Stage].UniformBuffers[UniformBufferIndex].EntriesRemappingInfo.AddDefaulted();
-		check(UBEntryRemappingIndex == ResourceEntryIndex);
-		FDescriptorSetRemappingInfo::FRemappingInfo& UBEntryRemapping = StageInfos[Stage].UniformBuffers[UniformBufferIndex].EntriesRemappingInfo[UBEntryRemappingIndex];
-		UBEntryRemapping.NewDescriptorSet = NewDescriptorSet;
-		UBEntryRemapping.NewBindingIndex = NewBindingIndex;
-
-		return NewBindingIndex;
-	}
-*/
-
-
 	static_assert(
 		sizeof(FRemappingInfo) == (sizeof(FRemappingInfo::NewDescriptorSet) + sizeof(FRemappingInfo::NewBindingIndex)),
 		"FRemappingInfo should not have padding! structure is used for MemCrc32/Memcmp");
@@ -786,15 +764,10 @@ public:
 		{
 		case FVulkanShaderHeader::UniformBuffer:
 			ensure(RemappingUBInfos[Stage][ParameterIndex].bHasConstantData);
-			//ensure(RemappingInfo->StageInfos[Stage].UniformBuffers[ParameterIndex].bHasConstantData);
-			//OutDescriptorSet = RemappingInfo->StageInfos[Stage].UniformBuffers[ParameterIndex].Remapping.NewDescriptorSet;
-			//OutBindingIndex = RemappingInfo->StageInfos[Stage].UniformBuffers[ParameterIndex].Remapping.NewBindingIndex;
 			OutDescriptorSet = RemappingUBInfos[Stage][ParameterIndex].Remapping.NewDescriptorSet;
 			OutBindingIndex = RemappingUBInfos[Stage][ParameterIndex].Remapping.NewBindingIndex;
 			break;
 		case FVulkanShaderHeader::Global:
-			//OutDescriptorSet = RemappingInfo->StageInfos[Stage].Globals[ParameterIndex].NewDescriptorSet;
-			//OutBindingIndex = RemappingInfo->StageInfos[Stage].Globals[ParameterIndex].NewBindingIndex;
 			OutDescriptorSet = RemappingGlobalInfos[Stage][ParameterIndex].NewDescriptorSet;
 			OutBindingIndex = RemappingGlobalInfos[Stage][ParameterIndex].NewBindingIndex;
 			break;
@@ -807,7 +780,6 @@ public:
 
 	inline const TArray<FDescriptorSetRemappingInfo::FRemappingInfo>& GetGlobalRemappingInfo(ShaderStage::EStage Stage) const
 	{
-		//OutDescriptorSet = RemappingUBInfos[Stage][ParameterIndex].Remapping.NewDescriptorSet;
 		return RemappingInfo->StageInfos[Stage].Globals;
 	}
 
diff --git a/Engine/Source/Runtime/VulkanRHI/Private/VulkanPendingState.h b/Engine/Source/Runtime/VulkanRHI/Private/VulkanPendingState.h
--- a/Engine/Source/Runtime/VulkanRHI/Private/VulkanPendingState.h
+++ b/Engine/Source/Runtime/VulkanRHI/Private/VulkanPendingState.h
@@ -128,11 +128,6 @@ public:
 		CurrentState->SetPackedGlobalShaderParameter(BufferIndex, Offset, NumBytes, NewValue);
 	}
 
-	inline void SetUniformBufferConstantData(uint32 BindingIndex, const TArray<uint8>& ConstantData)
-	{
-		CurrentState->SetUniformBufferConstantData(BindingIndex, ConstantData);
-	}
-
 	inline void SetSamplerStateForUBResource(uint32 DescriptorSet, uint32 BindingIndex, FVulkanSamplerState* Sampler)
 	{
 		CurrentState->SetSamplerState(DescriptorSet, BindingIndex, Sampler);
@@ -293,11 +288,6 @@ public:
 		CurrentState->SetTexture(DescriptorSet, BindingIndex, Texture, Layout);
 	}
 
-	inline void SetUniformBufferConstantData(ShaderStage::EStage Stage, uint32 BindingIndex, const TArray<uint8>& ConstantData)
-	{
-		CurrentState->SetUniformBufferConstantData(Stage, BindingIndex, ConstantData);
-	}
-
 	template<bool bDynamic>
 	inline void SetUniformBuffer(uint8 DescriptorSet, uint32 BindingIndex, const FVulkanUniformBuffer* UniformBuffer)
 	{
diff --git a/Engine/Source/Runtime/VulkanRHI/Private/VulkanPipelineState.h b/Engine/Source/Runtime/VulkanRHI/Private/VulkanPipelineState.h
--- a/Engine/Source/Runtime/VulkanRHI/Private/VulkanPipelineState.h
+++ b/Engine/Source/Runtime/VulkanRHI/Private/VulkanPipelineState.h
@@ -182,11 +182,6 @@ public:
 		PackedUniformBuffers.SetPackedGlobalParameter(BufferIndex, ByteOffset, NumBytes, NewValue, PackedUniformBuffersDirty);
 	}
 
-	inline void SetUniformBufferConstantData(uint32 BindingIndex, const TArray<uint8>& ConstantData)
-	{
-		PackedUniformBuffers.SetEmulatedUniformBufferIntoPacked(BindingIndex, ConstantData, PackedUniformBuffersDirty);
-	}
-
 	bool UpdateDescriptorSets(FVulkanCommandListContext* CmdListContext, FVulkanCmdBuffer* CmdBuffer)
 	{
 		check(!bUseBindless);
@@ -245,11 +240,6 @@ public:
 		PackedUniformBuffers[Stage].SetPackedGlobalParameter(BufferIndex, ByteOffset, NumBytes, NewValue, PackedUniformBuffersDirty[Stage]);
 	}
 
-	inline void SetUniformBufferConstantData(uint8 Stage, uint32 BindingIndex, const TArray<uint8>& ConstantData)
-	{
-		PackedUniformBuffers[Stage].SetEmulatedUniformBufferIntoPacked(BindingIndex, ConstantData, PackedUniformBuffersDirty[Stage]);
-	}
-
 	bool UpdateDescriptorSets(FVulkanCommandListContext* CmdListContext, FVulkanCmdBuffer* CmdBuffer)
 	{
 		check(!bUseBindless);
diff --git a/Engine/Source/Runtime/VulkanRHI/Private/VulkanResources.h b/Engine/Source/Runtime/VulkanRHI/Private/VulkanResources.h
--- a/Engine/Source/Runtime/VulkanRHI/Private/VulkanResources.h
+++ b/Engine/Source/Runtime/VulkanRHI/Private/VulkanResources.h
@@ -1292,8 +1292,6 @@ public:
 		}
 
 		OutPackedUniformBufferStagingMask = ((uint64)1 << (uint64)InCodeHeader.PackedUBs.Num()) - 1;
-		EmulatedUBsCopyInfo = InCodeHeader.EmulatedUBsCopyInfo;
-		EmulatedUBsCopyRanges = InCodeHeader.EmulatedUBCopyRanges;
 	}
 
 	inline void SetPackedGlobalParameter(uint32 BufferIndex, uint32 ByteOffset, uint32 NumBytes, const void* RESTRICT NewValue, uint64& InOutPackedUniformBufferStagingDirty)
@@ -1317,31 +1315,6 @@ public:
 	// Copies a 'real' constant buffer into the packed globals uniform buffer (only the used ranges)
 	inline void SetEmulatedUniformBufferIntoPacked(uint32 BindPoint, const TArray<uint8>& ConstantData, uint64& NEWPackedUniformBufferStagingDirty)
 	{
-		// Emulated UBs. Assumes UniformBuffersCopyInfo table is sorted by CopyInfo.SourceUBIndex
-		if (BindPoint < (uint32)EmulatedUBsCopyRanges.Num())
-		{
-			uint32 Range = EmulatedUBsCopyRanges[BindPoint];
-			uint16 Start = (Range >> 16) & 0xffff;
-			uint16 Count = Range & 0xffff;
-			const uint8* RESTRICT SourceData = ConstantData.GetData();
-			for (int32 Index = Start; Index < Start + Count; ++Index)
-			{
-				const CrossCompiler::FUniformBufferCopyInfo& CopyInfo = EmulatedUBsCopyInfo[Index];
-				check(CopyInfo.SourceUBIndex == BindPoint);
-				FPackedBuffer& StagingBuffer = PackedUniformBuffers[(int32)CopyInfo.DestUBIndex];
-				//check(ByteOffset + NumBytes <= (uint32)StagingBuffer.Num());
-				bool bChanged = false;
-				uint32* RESTRICT RawDst = (uint32*)(StagingBuffer.GetData() + CopyInfo.DestOffsetInFloats * 4);
-				uint32* RESTRICT RawSrc = (uint32*)(SourceData + CopyInfo.SourceOffsetInFloats * 4);
-				uint32* RESTRICT RawSrcEnd = RawSrc + CopyInfo.SizeInFloats;
-				do
-				{
-					bChanged |= CopyAndReturnNotEqual(*RawDst++, *RawSrc++);
-				}
-				while (RawSrc != RawSrcEnd);
-				NEWPackedUniformBufferStagingDirty = NEWPackedUniformBufferStagingDirty | ((uint64)(bChanged ? 1 : 0) << (uint64)CopyInfo.DestUBIndex);
-			}
-		}
 	}
 
 	inline const FPackedBuffer& GetBuffer(int32 Index) const
@@ -1351,10 +1324,6 @@ public:
 
 protected:
 	TArray<FPackedBuffer>									PackedUniformBuffers;
-
-	// Copies to Shader Code Header (shaders may be deleted when we use this object again)
-	TArray<CrossCompiler::FUniformBufferCopyInfo>			EmulatedUBsCopyInfo;
-	TArray<uint32>											EmulatedUBsCopyRanges;
 };
 
 class FVulkanStagingBuffer : public FRHIStagingBuffer
diff --git a/Engine/Source/Runtime/VulkanRHI/Private/VulkanShaders.cpp b/Engine/Source/Runtime/VulkanRHI/Private/VulkanShaders.cpp
--- a/Engine/Source/Runtime/VulkanRHI/Private/VulkanShaders.cpp
+++ b/Engine/Source/Runtime/VulkanRHI/Private/VulkanShaders.cpp
@@ -812,7 +812,7 @@ void FVulkanDescriptorSetsLayoutInfo::FinalizeBindings(const FVulkanDevice& Devi
 	{
 		if (const FVulkanShaderHeader* ShaderHeader = UBGatherInfo.CodeHeaders[Stage])
 		{
-			VkShaderStageFlags StageFlags = UEFrequencyToVKStageBit(bIsCompute ? SF_Compute : ShaderStage::GetFrequencyForGfxStage((ShaderStage::EStage)Stage));
+			const VkShaderStageFlags StageFlags = UEFrequencyToVKStageBit(bIsCompute ? SF_Compute : ShaderStage::GetFrequencyForGfxStage((ShaderStage::EStage)Stage));
 			Binding.stageFlags = StageFlags;
 
 			RemappingInfo.StageInfos[Stage].PackedUBBindingIndices.Reserve(ShaderHeader->PackedUBs.Num());
@@ -867,32 +867,13 @@ void FVulkanDescriptorSetsLayoutInfo::FinalizeBindings(const FVulkanDevice& Devi
 			{
 				const FVulkanShaderHeader::FGlobalInfo& GlobalInfo = ShaderHeader->Globals[Index];
 				int32 DescriptorSet = FindOrAddDescriptorSet(Stage);
-				VkDescriptorType Type = BindingToDescriptorType(ShaderHeader->GlobalDescriptorTypes[GlobalInfo.TypeIndex]);
-				uint16 CombinedSamplerStateAlias = GlobalInfo.CombinedSamplerStateAliasIndex;
-				uint32 NewBindingIndex = RemappingInfo.AddGlobal(Stage, Index, DescriptorSet, Type, CombinedSamplerStateAlias);
+				const VkDescriptorType Type = (VkDescriptorType)GlobalInfo.TypeIndex;
+				uint32 NewBindingIndex = RemappingInfo.AddGlobal(Stage, Index, DescriptorSet, Type);
 				Binding.binding = NewBindingIndex;
 				Binding.descriptorType = Type;
-				if (CombinedSamplerStateAlias == UINT16_MAX)
-				{
-					if (GlobalInfo.bImmutableSampler)
-					{
-						if (CurrentImmutableSampler < ImmutableSamplers.Num())
-						{
-							FVulkanSamplerState* SamplerState = ResourceCast(ImmutableSamplers[CurrentImmutableSampler]);
-							if (SamplerState && SamplerState->Sampler != VK_NULL_HANDLE)
-							{
-								Binding.pImmutableSamplers = &SamplerState->Sampler;
-							}
-							++CurrentImmutableSampler;
-						}
-					}
-
 					AddDescriptor(DescriptorSet, Binding);
 				}
 
-				Binding.pImmutableSamplers = nullptr;
-			}
-
 			if (ShaderHeader->InputAttachments.Num())
 			{
 				int32 DescriptorSet = FindOrAddDescriptorSet(Stage);
@@ -901,7 +882,7 @@ void FVulkanDescriptorSetsLayoutInfo::FinalizeBindings(const FVulkanDevice& Devi
 				{
 					int32 OriginalGlobalIndex = ShaderHeader->InputAttachments[SrcIndex].GlobalIndex;
 					const FVulkanShaderHeader::FGlobalInfo& OriginalGlobalInfo = ShaderHeader->Globals[OriginalGlobalIndex];
-					check(BindingToDescriptorType(ShaderHeader->GlobalDescriptorTypes[OriginalGlobalInfo.TypeIndex]) == VK_DESCRIPTOR_TYPE_INPUT_ATTACHMENT);
+					check((VkDescriptorType)OriginalGlobalInfo.TypeIndex == VK_DESCRIPTOR_TYPE_INPUT_ATTACHMENT);
 					int32 RemappingIndex = RemappingInfo.InputAttachmentData.AddDefaulted();
 					FInputAttachmentData& AttachmentData = RemappingInfo.InputAttachmentData[RemappingIndex];
 					AttachmentData.BindingIndex = RemappingInfo.StageInfos[Stage].Globals[OriginalGlobalIndex].NewBindingIndex;
diff --git a/Engine/Source/Runtime/VulkanRHI/Public/VulkanCommon.h b/Engine/Source/Runtime/VulkanRHI/Public/VulkanCommon.h
--- a/Engine/Source/Runtime/VulkanRHI/Public/VulkanCommon.h
+++ b/Engine/Source/Runtime/VulkanRHI/Public/VulkanCommon.h
@@ -26,25 +26,17 @@ namespace ShaderStage
 		// Keep the values in sync with EShaderFrequency
 		Vertex = 0,
 		Pixel = 1,
-
-#if VULKAN_SUPPORTS_GEOMETRY_SHADERS
 		Geometry = 2,
-#endif
 
 		RayGen = 3,
 		RayMiss = 4,
 		RayHitGroup = 5,
 		RayCallable = 6,
 
-#if VULKAN_SUPPORTS_GEOMETRY_SHADERS
-		NumGeometryStages = 1,
-#else
-		NumGeometryStages = 0,
-#endif
-
+		NumGraphicsStages = 3,
 		NumRayTracingStages = 4,
 
-		NumStages = (2 + NumGeometryStages + NumRayTracingStages),
+		NumStages = (NumGraphicsStages + NumRayTracingStages),
 
 		// Compute is its own pipeline, so it can all live as set 0
 		Compute = 0,
@@ -60,9 +52,7 @@ namespace ShaderStage
 		{
 		case SF_Vertex:		return Vertex;
 		case SF_Pixel:		return Pixel;
-#if VULKAN_SUPPORTS_GEOMETRY_SHADERS
 		case SF_Geometry:	return Geometry;
-#endif
 		case SF_RayGen:			return RayGen;
 		case SF_RayMiss:		return RayMiss;
 		case SF_RayHitGroup:	return RayHitGroup;
@@ -82,9 +72,7 @@ namespace ShaderStage
 		{
 		case EStage::Vertex:	return SF_Vertex;
 		case EStage::Pixel:		return SF_Pixel;
-#if VULKAN_SUPPORTS_GEOMETRY_SHADERS
 		case EStage::Geometry:	return SF_Geometry;
-#endif
 		case EStage::RayGen:		return SF_RayGen;
 		case EStage::RayMiss:		return SF_RayMiss;
 		case EStage::RayHitGroup:	return SF_RayHitGroup;
@@ -123,59 +111,6 @@ namespace VulkanBindless
 	};
 };
 
-namespace EVulkanBindingType
-{
-	enum EType : uint8
-	{
-		PackedUniformBuffer,	//VK_DESCRIPTOR_TYPE_UNIFORM_BUFFER
-		UniformBuffer,			//VK_DESCRIPTOR_TYPE_UNIFORM_BUFFER
-
-		CombinedImageSampler,	//VK_DESCRIPTOR_TYPE_COMBINED_IMAGE_SAMPLER	*not used*
-		Sampler,				//VK_DESCRIPTOR_TYPE_SAMPLER				(HLSL: SamplerState/SamplerComparisonState)
-		Image,					//VK_DESCRIPTOR_TYPE_SAMPLED_IMAGE			(HLSL: Texture2D/3D/Cube)
-
-		UniformTexelBuffer,		//VK_DESCRIPTOR_TYPE_UNIFORM_TEXEL_BUFFER	(HLSL: Buffer)
-
-		// A storage image is a descriptor type that is used for load, store, and atomic operations on image memory from within shaders bound to pipelines.
-		StorageImage,			//VK_DESCRIPTOR_TYPE_STORAGE_IMAGE			(HLSL: RWTexture2D/3D/Cube)
-
-		//A storage texel buffer represents a tightly packed array of homogeneous formatted data that is stored in a buffer and is made accessible to shaders. Storage texel buffers differ from uniform texel buffers in that they support stores and atomic operations in shaders, may support a different maximum length, and may have different performance characteristics.
-		StorageTexelBuffer,		//VK_DESCRIPTOR_TYPE_STORAGE_TEXEL_BUFFER	(HLSL: RWBuffer)
-
-		// A storage buffer is a region of structured storage that supports both read and write access for shaders. In addition to general read and write operations, some members of storage buffers can be used as the target of atomic operations. In general, atomic operations are only supported on members that have unsigned integer formats.
-		StorageBuffer,			//VK_DESCRIPTOR_TYPE_STORAGE_BUFFER			(HLSL: StructuredBuffer/RWStructureBuffer/ByteAddressBuffer/RWByteAddressBuffer)
-
-		InputAttachment,
-
-		AccelerationStructure,	//VK_DESCRIPTOR_TYPE_ACCELERATION_STRUCTURE_KHR
-
-		Count,
-	};
-
-	static inline char GetBindingTypeChar(EType Type)
-	{
-		// Make sure these do NOT alias EPackedTypeName*
-		switch (Type)
-		{
-		case UniformBuffer:			return 'b';
-		case CombinedImageSampler:	return 'c';
-		case Sampler:				return 'p';
-		case Image:					return 'w';
-		case UniformTexelBuffer:	return 'x';
-		case StorageImage:			return 'y';
-		case StorageTexelBuffer:	return 'z';
-		case StorageBuffer:			return 'v';
-		case InputAttachment:		return 'a';
-		case AccelerationStructure:	return 'r';
-		default:
-			check(0);
-			break;
-		}
-
-		return 0;
-	}
-}
-
 DECLARE_LOG_CATEGORY_EXTERN(LogVulkan, Display, All);
 
 template< class T >
diff --git a/Engine/Source/Runtime/VulkanRHI/Public/VulkanShaderResources.h b/Engine/Source/Runtime/VulkanRHI/Public/VulkanShaderResources.h
--- a/Engine/Source/Runtime/VulkanRHI/Public/VulkanShaderResources.h
+++ b/Engine/Source/Runtime/VulkanRHI/Public/VulkanShaderResources.h
@@ -11,52 +11,6 @@
 #include "VulkanCommon.h"
 #include "VulkanThirdParty.h"
 
-static inline VkDescriptorType BindingToDescriptorType(EVulkanBindingType::EType Type)
-{
-	// Make sure these do NOT alias EPackedTypeName*
-	switch (Type)
-	{
-	case EVulkanBindingType::PackedUniformBuffer:	return VK_DESCRIPTOR_TYPE_UNIFORM_BUFFER;
-	case EVulkanBindingType::UniformBuffer:			return VK_DESCRIPTOR_TYPE_UNIFORM_BUFFER;
-	case EVulkanBindingType::CombinedImageSampler:	return VK_DESCRIPTOR_TYPE_COMBINED_IMAGE_SAMPLER;
-	case EVulkanBindingType::Sampler:				return VK_DESCRIPTOR_TYPE_SAMPLER;
-	case EVulkanBindingType::Image:					return VK_DESCRIPTOR_TYPE_SAMPLED_IMAGE;
-	case EVulkanBindingType::UniformTexelBuffer:	return VK_DESCRIPTOR_TYPE_UNIFORM_TEXEL_BUFFER;
-	case EVulkanBindingType::StorageImage:			return VK_DESCRIPTOR_TYPE_STORAGE_IMAGE;
-	case EVulkanBindingType::StorageTexelBuffer:	return VK_DESCRIPTOR_TYPE_STORAGE_TEXEL_BUFFER;
-	case EVulkanBindingType::StorageBuffer:			return VK_DESCRIPTOR_TYPE_STORAGE_BUFFER;
-	case EVulkanBindingType::InputAttachment:		return VK_DESCRIPTOR_TYPE_INPUT_ATTACHMENT;
-	case EVulkanBindingType::AccelerationStructure:	return VK_DESCRIPTOR_TYPE_ACCELERATION_STRUCTURE_KHR;
-	default:
-		check(0);
-		break;
-	}
-
-	return VK_DESCRIPTOR_TYPE_MAX_ENUM;
-}
-
-static inline EVulkanBindingType::EType DescriptorTypeToBinding(VkDescriptorType Type, bool bUsePacked = false)
-{
-	switch (Type)
-	{
-	case VK_DESCRIPTOR_TYPE_UNIFORM_BUFFER:				return bUsePacked ? EVulkanBindingType::PackedUniformBuffer : EVulkanBindingType::UniformBuffer;
-	case VK_DESCRIPTOR_TYPE_COMBINED_IMAGE_SAMPLER:		return EVulkanBindingType::CombinedImageSampler;
-	case VK_DESCRIPTOR_TYPE_SAMPLER:					return EVulkanBindingType::Sampler;
-	case VK_DESCRIPTOR_TYPE_SAMPLED_IMAGE:				return EVulkanBindingType::Image;
-	case VK_DESCRIPTOR_TYPE_UNIFORM_TEXEL_BUFFER:		return EVulkanBindingType::UniformTexelBuffer;
-	case VK_DESCRIPTOR_TYPE_STORAGE_IMAGE:				return EVulkanBindingType::StorageImage;
-	case VK_DESCRIPTOR_TYPE_STORAGE_TEXEL_BUFFER:		return EVulkanBindingType::StorageTexelBuffer;
-	case VK_DESCRIPTOR_TYPE_STORAGE_BUFFER:				return EVulkanBindingType::StorageBuffer;
-	case VK_DESCRIPTOR_TYPE_INPUT_ATTACHMENT:			return EVulkanBindingType::InputAttachment;
-	case VK_DESCRIPTOR_TYPE_ACCELERATION_STRUCTURE_KHR: return EVulkanBindingType::AccelerationStructure;
-	default:
-		check(0);
-		break;
-	}
-
-	return EVulkanBindingType::Count;
-}
-
 
 // Vulkan ParameterMap:
 // Buffer Index = EBufferIndex
@@ -86,27 +40,14 @@ struct FVulkanShaderHeader
 		uint32	BindingIndexOffset = UINT32_MAX;
 	};
 
-	struct FUBResourceInfo
-	{
-		uint16									SourceUBResourceIndex;
-		uint16									OriginalBindingIndex;
-		// Index into the Global Array
-		uint16									GlobalIndex;
-		TEnumAsByte<EUniformBufferBaseType>		UBBaseType;
-		uint8									Pad0 = 0;
-#if VULKAN_ENABLE_BINDING_DEBUG_NAMES
-		FString									DebugName;
-#endif
-	};
+
 
 	struct FUniformBufferInfo
 	{
 		uint32					LayoutHash;
 		uint16					ConstantDataOriginalBindingIndex;
 		uint8					bOnlyHasResources;
-		uint8					Pad0 = 0;
-		//uint32					ConstantDataSizeInBytes;
-		TArray<FUBResourceInfo>	ResourceEntries;
+		uint8					bHasResources;
 #if VULKAN_ENABLE_BINDING_DEBUG_NAMES
 		FString					DebugName;
 #endif
@@ -115,38 +56,27 @@ struct FVulkanShaderHeader
 
 	struct FGlobalInfo
 	{
-		uint16							OriginalBindingIndex;
-		// If this is UINT16_MAX, it's a regular parameter, otherwise this is the SamplerState portion for a CombinedImageSampler
-		// and this is the index into Global for the Texture portion
-		uint16							CombinedSamplerStateAliasIndex;
-		uint16							TypeIndex;
-		// 1 if this is an immutable sampler
-		uint8							bImmutableSampler = 0;
-		uint8							Pad0 = 0;
+		uint32							TypeIndex;
 #if VULKAN_ENABLE_BINDING_DEBUG_NAMES
 		FString							DebugName;
 #endif
 	};
-	TArray<FGlobalInfo>						Globals;
-	TArray<TEnumAsByte<EVulkanBindingType::EType>>	GlobalDescriptorTypes;
+	TArray<FGlobalInfo>					Globals;
 
 	struct FPackedGlobalInfo
 	{
 		uint16							ConstantDataSizeInFloats;
-		CrossCompiler::EPackedTypeIndex	PackedTypeIndex;
 		uint8							PackedUBIndex;
+		uint8							Pad0 = 0;
 #if VULKAN_ENABLE_BINDING_DEBUG_NAMES
 		FString							DebugName;
 #endif
 	};
-	TArray<FPackedGlobalInfo>	PackedGlobals;
+	TArray<FPackedGlobalInfo>			PackedGlobals;
 
 	struct FPackedUBInfo
 	{
 		uint32							SizeInBytes;
-		uint16							OriginalBindingIndex;
-		CrossCompiler::EPackedTypeIndex	PackedTypeIndex;
-		uint8							Pad0 = 0;
 		uint32							SPIRVDescriptorSetOffset;
 		uint32							SPIRVBindingIndexOffset;
 	};
@@ -174,10 +104,6 @@ struct FVulkanShaderHeader
 	};
 	TArray<FInputAttachment>				InputAttachments;
 
-	// Number of copies per emulated buffer source index (to skip searching among UniformBuffersCopyInfo). Upper uint16 is the index, Lower uint16 is the count
-	TArray<uint32>									EmulatedUBCopyRanges;
-	TArray<CrossCompiler::FUniformBufferCopyInfo>	EmulatedUBsCopyInfo;
-
 	// Mostly relevant for Vertex Shaders
 	uint32									InOutMask;
 
@@ -229,25 +155,12 @@ inline FArchive& operator<<(FArchive& Ar, FVulkanShaderHeader::FSpirvInfo& Info)
 	return Ar;
 }
 
-inline FArchive& operator<<(FArchive& Ar, FVulkanShaderHeader::FUBResourceInfo& Entry)
-{
-	Ar << Entry.SourceUBResourceIndex;
-	Ar << Entry.OriginalBindingIndex;
-	Ar << Entry.GlobalIndex;
-	Ar << Entry.UBBaseType;
-#if VULKAN_ENABLE_BINDING_DEBUG_NAMES
-	Ar << Entry.DebugName;
-#endif
-	return Ar;
-}
-
 inline FArchive& operator<<(FArchive& Ar, FVulkanShaderHeader::FUniformBufferInfo& UBInfo)
 {
 	Ar << UBInfo.LayoutHash;
 	Ar << UBInfo.ConstantDataOriginalBindingIndex;
 	Ar << UBInfo.bOnlyHasResources;
-	//Ar << UBInfo.ConstantDataSizeInBytes;
-	Ar << UBInfo.ResourceEntries;
+	Ar << UBInfo.bHasResources;
 #if VULKAN_ENABLE_BINDING_DEBUG_NAMES
 	Ar << UBInfo.DebugName;
 #endif
@@ -257,7 +170,6 @@ inline FArchive& operator<<(FArchive& Ar, FVulkanShaderHeader::FUniformBufferInf
 inline FArchive& operator<<(FArchive& Ar, FVulkanShaderHeader::FPackedGlobalInfo& PackedGlobalInfo)
 {
 	Ar << PackedGlobalInfo.ConstantDataSizeInFloats;
-	Ar << PackedGlobalInfo.PackedTypeIndex;
 	Ar << PackedGlobalInfo.PackedUBIndex;
 #if VULKAN_ENABLE_BINDING_DEBUG_NAMES
 	Ar << PackedGlobalInfo.DebugName;
@@ -268,8 +180,6 @@ inline FArchive& operator<<(FArchive& Ar, FVulkanShaderHeader::FPackedGlobalInfo
 inline FArchive& operator<<(FArchive& Ar, FVulkanShaderHeader::FPackedUBInfo& PackedUBInfo)
 {
 	Ar << PackedUBInfo.SizeInBytes;
-	Ar << PackedUBInfo.OriginalBindingIndex;
-	Ar << PackedUBInfo.PackedTypeIndex;
 	Ar << PackedUBInfo.SPIRVDescriptorSetOffset;
 	Ar << PackedUBInfo.SPIRVBindingIndexOffset;
 	return Ar;
@@ -277,10 +187,7 @@ inline FArchive& operator<<(FArchive& Ar, FVulkanShaderHeader::FPackedUBInfo& Pa
 
 inline FArchive& operator<<(FArchive& Ar, FVulkanShaderHeader::FGlobalInfo& GlobalInfo)
 {
-	Ar << GlobalInfo.OriginalBindingIndex;
-	Ar << GlobalInfo.CombinedSamplerStateAliasIndex;
 	Ar << GlobalInfo.TypeIndex;
-	Ar << GlobalInfo.bImmutableSampler;
 #if VULKAN_ENABLE_BINDING_DEBUG_NAMES
 	Ar << GlobalInfo.DebugName;
 #endif
@@ -298,12 +205,9 @@ inline FArchive& operator<<(FArchive& Ar, FVulkanShaderHeader& Header)
 {
 	Ar << Header.UniformBuffers;
 	Ar << Header.Globals;
-	Ar << Header.GlobalDescriptorTypes;
 	Ar << Header.PackedGlobals;
 	Ar << Header.PackedUBs;
 	Ar << Header.InputAttachments;
-	Ar << Header.EmulatedUBCopyRanges;
-	Ar << Header.EmulatedUBsCopyInfo;
 	Ar << Header.InOutMask;
 	Ar << Header.RayTracingPayloadType;
 	Ar << Header.RayTracingPayloadSize;
