VulkanRHI: Remove unused data in shader header.  Remove FVulkanShaderHeader::EType and use EShaderParameterType directly instead, removes the need for an extra array in shader compiler.

#rnx

diff --git a/Engine/Shaders/Public/Platform/Vulkan/VulkanCommon.ush b/Engine/Shaders/Public/Platform/Vulkan/VulkanCommon.ush
--- a/Engine/Shaders/Public/Platform/Vulkan/VulkanCommon.ush
+++ b/Engine/Shaders/Public/Platform/Vulkan/VulkanCommon.ush
@@ -7,7 +7,7 @@
 #pragma once
 
 // Update this GUID to invalidate shader recompilation for only Vulkan shaders
-#pragma message("UESHADERMETADATA_VERSION C0580A2D-7AA2-4CE4-9D2D-DFC9B838785A")
+#pragma message("UESHADERMETADATA_VERSION FE20F789-553C-4B7E-B4FA-423EDFA498B4")
 
 #define COMPILER_SUPPORTS_ULONG_TYPES 1
 #define COMPILER_SUPPORTS_UINT64_IMAGE_ATOMICS (1 && COMPILER_SUPPORTS_ULONG_TYPES)
diff --git a/Engine/Source/Developer/VulkanShaderFormat/Private/VulkanShaderCompiler.cpp b/Engine/Source/Developer/VulkanShaderFormat/Private/VulkanShaderCompiler.cpp
--- a/Engine/Source/Developer/VulkanShaderFormat/Private/VulkanShaderCompiler.cpp
+++ b/Engine/Source/Developer/VulkanShaderFormat/Private/VulkanShaderCompiler.cpp
@@ -542,7 +542,7 @@ static void AddUniformBuffer(
 	}
 
 	InOutParameterMap.RemoveParameterAllocation(*UBName);
-	InOutParameterMap.AddParameterAllocation(*UBName, HeaderUBIndex, (uint16)FVulkanShaderHeader::UniformBuffer, 1, EShaderParameterType::UniformBuffer);
+	InOutParameterMap.AddParameterAllocation(*UBName, HeaderUBIndex, 0, 1, EShaderParameterType::UniformBuffer);
 }
 
 static int32 DoAddGlobal(const FString& Name, FVulkanShaderHeader& OutHeader, TArray<FString>& OutGlobalNames)
@@ -589,27 +589,17 @@ static void PrepareUBResourceEntryGlobals(const TArray<uint32>& BindingArray, co
 }
 
 
-static void PrepareGlobals(const VulkanShaderCompilerSerializedOutput& SerializedOutput, const FSpirvReflectBindings& SpirvReflectBindings, const FShaderResourceTable& SRT, const TMap<FString, FVulkanShaderHeader::EType>& EntryTypes, const FShaderCompilerInput& ShaderInput, const TArray<FString>& ParameterNames, FShaderParameterMap& ParameterMap, TArray<FString>& OutGlobalNames, FVulkanShaderHeader& OutHeader)
+static void PrepareGlobals(const VulkanShaderCompilerSerializedOutput& SerializedOutput, const FSpirvReflectBindings& SpirvReflectBindings, const FShaderResourceTable& SRT, const FShaderCompilerInput& ShaderInput, const TArray<FString>& ParameterNames, FShaderParameterMap& ParameterMap, TArray<FString>& OutGlobalNames, FVulkanShaderHeader& OutHeader)
 {
-	auto IsSamplerState = [&SpirvReflectBindings](const FString& ParameterName)
-	{
-		for (SpvReflectDescriptorBinding* DescriptorBinding : SpirvReflectBindings.Samplers)
-		{
-			if (ParameterName == DescriptorBinding->name)
-			{
-				return true;
-			}
-		}
-		return false;
-	};
-
 	// First pass, gather names for all the Globals that are NOT Samplers
+	for (int32 ParameterIndex = 0; ParameterIndex < ParameterNames.Num(); ++ParameterIndex)
 	{
-		auto AddGlobalNamesForUB = [&](const FString& ParameterName)
-		{
-			TOptional<FParameterAllocation> ParameterAllocation = ParameterMap.FindParameterAllocation(*ParameterName);
-			checkf(ParameterAllocation.IsSet(), TEXT("PrepareGlobals failed to find resource ParameterName=%s"), *ParameterName);
+		const FString& ParameterName = ParameterNames[ParameterIndex];
+		TOptional<FParameterAllocation> ParameterAllocation = ParameterMap.FindParameterAllocation(*ParameterName);
+		checkf(ParameterAllocation.IsSet(), TEXT("PrepareGlobals failed to find resource ParameterName=%s"), *ParameterName);
 
+		if (ParameterAllocation->Type == EShaderParameterType::UniformBuffer)
+		{
 			// Add used resources...
 			if (SRT.ResourceTableBits & (1 << ParameterAllocation->BufferIndex))
 			{
@@ -617,82 +607,35 @@ static void PrepareGlobals(const VulkanShaderCompilerSerializedOutput& Serialize
 				PrepareUBResourceEntryGlobals(SRT.ShaderResourceViewMap, ShaderInput.Environment.ResourceTableMap, ParameterAllocation->BufferIndex, ParameterName, OutGlobalNames, OutHeader);
 				PrepareUBResourceEntryGlobals(SRT.UnorderedAccessViewMap, ShaderInput.Environment.ResourceTableMap, ParameterAllocation->BufferIndex, ParameterName, OutGlobalNames, OutHeader);
 			}
-		};
-
-		for (int32 ParameterIndex = 0; ParameterIndex < ParameterNames.Num(); ++ParameterIndex)
+		}
+		else if ((ParameterAllocation->Type == EShaderParameterType::LooseData) || IsParameterBindless(ParameterAllocation->Type))
 		{
-			const FString& ParameterName = ParameterNames[ParameterIndex];
-			const FVulkanShaderHeader::EType* FoundType = EntryTypes.Find(ParameterName);
-			if (FoundType)
-			{
-				switch (*FoundType)
-				{
-				case FVulkanShaderHeader::Global:
-					if (!IsSamplerState(ParameterName))
-					{
-						DoAddGlobal(ParameterName, OutHeader, OutGlobalNames);
-					}
-					break;
-				case FVulkanShaderHeader::UniformBuffer:
-					AddGlobalNamesForUB(ParameterName);
-					break;
-				case FVulkanShaderHeader::PackedGlobal:
-					// Ignore
-					break;
-				default:
-					check(0);
-					break;
-				}
-			}
-			else
-			{
-				AddGlobalNamesForUB(ParameterName);
-			}
+			// Do nothing
+		}
+		else if (ParameterAllocation->Type != EShaderParameterType::Sampler)
+		{
+			DoAddGlobal(ParameterName, OutHeader, OutGlobalNames);
 		}
 	}
 
 	// Second pass, add all samplers
+	for (int32 ParameterIndex = 0; ParameterIndex < ParameterNames.Num(); ++ParameterIndex)
 	{
-		auto AddGlobalNamesForUB = [&](const FString& ParameterName)
+		const FString& ParameterName = ParameterNames[ParameterIndex];
+		TOptional<FParameterAllocation> ParameterAllocation = ParameterMap.FindParameterAllocation(*ParameterName);
+		checkf(ParameterAllocation.IsSet(), TEXT("PrepareGlobals failed to find resource ParameterName=%s"), *ParameterName);
+			
+		if (ParameterAllocation->Type == EShaderParameterType::UniformBuffer)
 		{
-			TOptional<FParameterAllocation> ParameterAllocation = ParameterMap.FindParameterAllocation(*ParameterName);
-			checkf(ParameterAllocation.IsSet(), TEXT("PrepareGlobals failed to find sampler ParameterName=%s"), *ParameterName);
-
-			// Add used resources...
+			// Add used samplers...
 			if (SRT.ResourceTableBits & (1 << ParameterAllocation->BufferIndex))
 			{
 				PrepareUBResourceEntryGlobals(SRT.SamplerMap, ShaderInput.Environment.ResourceTableMap, ParameterAllocation->BufferIndex, ParameterName, OutGlobalNames, OutHeader);
 			}
-		};
-
-		for (int32 ParameterIndex = 0; ParameterIndex < ParameterNames.Num(); ++ParameterIndex)
+		}
+		else if (ParameterAllocation->Type == EShaderParameterType::Sampler)
 		{
-			const FString& ParameterName = ParameterNames[ParameterIndex];
-			const FVulkanShaderHeader::EType* FoundType = EntryTypes.Find(ParameterName);
-			if (FoundType)
-			{
-				switch (*FoundType)
-				{
-				case FVulkanShaderHeader::Global:
-					if (IsSamplerState(ParameterName))
-					{
-						DoAddGlobal(ParameterName, OutHeader, OutGlobalNames);
-					}
-					break;
-				case FVulkanShaderHeader::UniformBuffer:
-					AddGlobalNamesForUB(ParameterName);
-					break;
-				case FVulkanShaderHeader::PackedGlobal:
-					break;
-				default:
-					check(0);
-					break;
-				}
-			}
-			else
-			{
-				AddGlobalNamesForUB(ParameterName);
-			}
+			DoAddGlobal(ParameterName, OutHeader, OutGlobalNames);
 		}
 	}
 
@@ -712,7 +655,6 @@ static void PrepareGlobals(const VulkanShaderCompilerSerializedOutput& Serialize
 
 static void ConvertToHeader(
 	FShaderResourceTable& ShaderResourceTable,
-	const TMap<FString, FVulkanShaderHeader::EType>& EntryTypes,
 	const FShaderCompilerInput& ShaderInput,
 	const FSpirvReflectBindings& SpirvReflectBindings,
 	FShaderParameterMap& InOutParameterMap,
@@ -727,58 +669,28 @@ static void ConvertToHeader(
 	TArray<FString> ParameterNames;
 	InOutParameterMap.GetAllParameterNames(ParameterNames);
 
-	PrepareGlobals(SerializedOutput, SpirvReflectBindings, ShaderResourceTable, EntryTypes, ShaderInput, ParameterNames, InOutParameterMap, GlobalNames, OutHeader);
+	PrepareGlobals(SerializedOutput, SpirvReflectBindings, ShaderResourceTable, ShaderInput, ParameterNames, InOutParameterMap, GlobalNames, OutHeader);
 
 	for (int32 ParameterIndex = 0; ParameterIndex < ParameterNames.Num(); ++ParameterIndex)
 	{
-		uint16 BufferIndex;
-		uint16 BaseIndex;
-		uint16 Size;
 		const FString& ParameterName = *ParameterNames[ParameterIndex];
-		const bool bFoundParam = InOutParameterMap.FindParameterAllocation(*ParameterName, BufferIndex, BaseIndex, Size);
-		check(bFoundParam);
+		const FParameterAllocation* ParameterAllocation = InOutParameterMap.GetParameterMap().Find(*ParameterName);
+		check(ParameterAllocation);
 
-		const FVulkanShaderHeader::EType* FoundType = EntryTypes.Find(ParameterName);
-		if (FoundType)
+		if (ParameterAllocation->Type == EShaderParameterType::UniformBuffer)
 		{
-			switch (*FoundType)
-			{
-			case FVulkanShaderHeader::Global:
-				{
-					const int32 HeaderGlobalIndex = AddGlobal(ParameterName, BaseIndex, SerializedOutput, GlobalNames);
-
-					const FParameterAllocation* ParameterAllocation = InOutParameterMap.GetParameterMap().Find(*ParameterName);
-					check(ParameterAllocation);
-					const EShaderParameterType ParamType = ParameterAllocation->Type;
-
-					InOutParameterMap.RemoveParameterAllocation(*ParameterName);
-					InOutParameterMap.AddParameterAllocation(*ParameterName, (uint16)FVulkanShaderHeader::Global, HeaderGlobalIndex, Size, ParamType);
-				}
-				break;
-			case FVulkanShaderHeader::PackedGlobal:
-				{
-					FVulkanShaderHeader::FPackedGlobalInfo& PackedGlobalInfo = OutHeader.PackedGlobals.AddZeroed_GetRef();
-					PackedGlobalInfo.PackedUBIndex = BufferIndex;
-					checkf(Size > 0, TEXT("Assertion failed for shader parameter: %s"), *ParameterName);
-					PackedGlobalInfo.ConstantDataSizeInFloats = Size / sizeof(float);
-#if VULKAN_ENABLE_BINDING_DEBUG_NAMES
-					PackedGlobalInfo.DebugName = ParameterName;
-#endif
-					// Keep the original parameter info from InOutParameterMap as it's a shortcut into the packed global array!
-				}
-				break;
-			case FVulkanShaderHeader::UniformBuffer:
-				AddUniformBuffer(ShaderResourceTable, ShaderInput, ParameterName, BufferIndex, InOutParameterMap, SerializedOutput, GlobalNames);
-				break;
-			default:
-				check(0);
-				break;
-			}
+			AddUniformBuffer(ShaderResourceTable, ShaderInput, ParameterName, ParameterAllocation->BufferIndex, InOutParameterMap, SerializedOutput, GlobalNames);
+		}
+		else if ((ParameterAllocation->Type == EShaderParameterType::LooseData) || IsParameterBindless(ParameterAllocation->Type))
+		{
+			// Do nothing
 		}
 		else
 		{
-			// Not found means it's a new resource-only UniformBuffer
-			AddUniformBuffer(ShaderResourceTable, ShaderInput, ParameterName, BufferIndex, InOutParameterMap, SerializedOutput, GlobalNames);
+			const int32 HeaderGlobalIndex = AddGlobal(ParameterName, ParameterAllocation->BaseIndex, SerializedOutput, GlobalNames);
+
+			InOutParameterMap.RemoveParameterAllocation(*ParameterName);
+			InOutParameterMap.AddParameterAllocation(*ParameterName, 0, HeaderGlobalIndex, ParameterAllocation->Size, ParameterAllocation->Type);
 		}
 	}
 
@@ -837,7 +749,6 @@ static void BuildShaderOutput(
 	FShaderCompilerOutput&		ShaderOutput,
 	const FVulkanShaderCompilerInternalState& InternalState,
 	const FSpirvReflectBindings& SpirvReflectBindings,
-	TMap<FString, FVulkanShaderHeader::EType>& EntryTypes,
 	const FString&				DebugName,
 	uint32						PackedGlobalArraySize,
 	TBitArray<>&				UsedUniformBufferSlots
@@ -981,7 +892,7 @@ static void BuildShaderOutput(
 		PackedUB.SPIRVBindingIndexOffset = Entry->WordBindingIndex;
 	}
 
-	ConvertToHeader(ShaderResourceTable, EntryTypes, ShaderInput, SpirvReflectBindings, ShaderOutput.ParameterMap, SerializedOutput);
+	ConvertToHeader(ShaderResourceTable, ShaderInput, SpirvReflectBindings, ShaderOutput.ParameterMap, SerializedOutput);
 
 	if (ShaderInput.Environment.CompilerFlags.Contains(CFLAG_ExtraShaderData))
 	{
@@ -1205,8 +1116,6 @@ static bool BuildShaderOutputFromSpirv(
 	const int32 MaxNumBits = VulkanBindless::MaxUniformBuffersPerStage * SF_NumFrequencies;
 	UsedUniformBufferSlots.Init(false, MaxNumBits);
 
-	TMap<FString, FVulkanShaderHeader::EType> EntryTypes;
-
 	// Final descriptor binding numbers for all other resource types
 	{
 		const ShaderStage::EStage UEStage = ShaderStage::GetStageForFrequency(InternalState.GetShaderFrequency());
@@ -1259,8 +1168,6 @@ static bool BuildShaderOutputFromSpirv(
 				case VK_DESCRIPTOR_TYPE_STORAGE_BUFFER:
 				case VK_DESCRIPTOR_TYPE_STORAGE_IMAGE:
 					HandleReflectedShaderUAV(ResourceName, BindingOffset, ReflectionSlot, 1, Output);
-					EntryTypes.Add(ResourceName, FVulkanShaderHeader::Global);
-
 					AddShaderValidationType(BindingOffset, ParsedParam, Output);
 					break;
 
@@ -1275,21 +1182,16 @@ static bool BuildShaderOutputFromSpirv(
 
 				case VK_DESCRIPTOR_TYPE_UNIFORM_TEXEL_BUFFER:
 					HandleReflectedShaderResource(ResourceName, BindingOffset, ReflectionSlot, 1, Output);
-					EntryTypes.Add(ResourceName, FVulkanShaderHeader::Global);
-
 					AddShaderValidationType(BindingOffset, ParsedParam, Output);
 					break;
 
 				case VK_DESCRIPTOR_TYPE_SAMPLER:
 					HandleReflectedShaderSampler(ResourceName, ReflectionSlot, Output);
 					//HandleReflectedShaderSampler(ResourceName, BindingOffset, ReflectionSlot, 1, Output);
-					EntryTypes.Add(ResourceName, FVulkanShaderHeader::Global);
 					break;
 
 				case VK_DESCRIPTOR_TYPE_ACCELERATION_STRUCTURE_KHR:
 					HandleReflectedShaderResource(ResourceName, BindingOffset, ReflectionSlot, 1, Output);
-					EntryTypes.Add(ResourceName, FVulkanShaderHeader::Global);
-
 					AddShaderValidationType(BindingOffset, ParsedParam, Output);
 					break;
 
@@ -1308,7 +1210,6 @@ static bool BuildShaderOutputFromSpirv(
 						{
 							check(ReflectionSlot == FShaderParametersMetadata::kRootCBufferBindingIndex);
 							HandleReflectedUniformBuffer(ResourceName, ReflectionSlot, Output);
-							EntryTypes.Add(ResourceName, FVulkanShaderHeader::UniformBuffer);
 						}
 
 						// Register all uniform buffer members of Globals as loose data
@@ -1332,8 +1233,6 @@ static bool BuildShaderOutputFromSpirv(
 									Member.size,
 									Output
 								);
-
-								EntryTypes.Add(FString(AdjustedMemberName), FVulkanShaderHeader::PackedGlobal);
 							}
 
 							PackedGlobalArraySize = FMath::Max<uint32>((Member.absolute_offset + Member.size), PackedGlobalArraySize);
@@ -1344,8 +1243,6 @@ static bool BuildShaderOutputFromSpirv(
 						check(!UsedUniformBufferSlots[ReflectionSlot]);
 						UsedUniformBufferSlots[ReflectionSlot] = true;
 						HandleReflectedUniformBuffer(ResourceName, ReflectionSlot, Output);
-						EntryTypes.Add(ResourceName, FVulkanShaderHeader::UniformBuffer);
-
 						AddShaderValidationUBSize(BindingOffset, Binding->block.padded_size, Output);
 					}
 					break;
@@ -1426,7 +1323,6 @@ static bool BuildShaderOutputFromSpirv(
 		Output,
 		InternalState,
 		Bindings,
-		EntryTypes,
 		InternalState.GetDebugName(),
 		PackedGlobalArraySize,
 		UsedUniformBufferSlots
@@ -2073,7 +1969,7 @@ static void UpdateBindlessUBs(const FVulkanShaderCompilerInternalState& Internal
 #endif
 
 			const int32 UBIndex = SerializedOutput.Header.UniformBuffers.Num() - 1;
-			Output.ParameterMap.AddParameterAllocation(*CBName, UBIndex, (uint16)FVulkanShaderHeader::UniformBuffer, 1, EShaderParameterType::UniformBuffer);
+			Output.ParameterMap.AddParameterAllocation(*CBName, UBIndex, 0, 1, EShaderParameterType::UniformBuffer);
 		}
 	}
 }
diff --git a/Engine/Source/Runtime/VulkanRHI/Private/VulkanCommands.cpp b/Engine/Source/Runtime/VulkanRHI/Private/VulkanCommands.cpp
--- a/Engine/Source/Runtime/VulkanRHI/Private/VulkanCommands.cpp
+++ b/Engine/Source/Runtime/VulkanRHI/Private/VulkanCommands.cpp
@@ -487,7 +487,7 @@ void FVulkanCommandListContext::RHISetUniformBufferDynamicOffset(FUniformBufferS
 			{
 				uint8 DescriptorSet;
 				uint32 BindingIndex;
-				if (DescriptorInfo.GetDescriptorSetAndBindingIndex(FVulkanShaderHeader::UniformBuffer, Stage, BufferIndex, DescriptorSet, BindingIndex))
+				DescriptorInfo.GetUBDescriptorSetAndBindingIndex(Stage, BufferIndex, DescriptorSet, BindingIndex);
 				{
 					// Uniform views always bind max supported range, so make sure Offset+Range is within buffer allocation
 					check((InOffset + PLATFORM_MAX_UNIFORM_BUFFER_RANGE) <= UniformBuffer->Allocation.Size);
@@ -522,10 +522,7 @@ void FVulkanCommandListContext::RHISetShaderUniformBuffer(FRHIGraphicsShader* Sh
 
 		uint8 DescriptorSet;
 		uint32 BindingIndex;
-		if (!DescriptorInfo.GetDescriptorSetAndBindingIndex(FVulkanShaderHeader::UniformBuffer, Stage, BufferIndex, DescriptorSet, BindingIndex))
-		{
-			return;
-		}
+		DescriptorInfo.GetUBDescriptorSetAndBindingIndex(Stage, BufferIndex, DescriptorSet, BindingIndex);
 
 		const VkDescriptorType DescriptorType = DescriptorInfo.GetDescriptorType(DescriptorSet, BindingIndex);
 
@@ -580,10 +577,7 @@ void FVulkanCommandListContext::RHISetShaderUniformBuffer(FRHIComputeShader* Com
 		
 		uint8 DescriptorSet;
 		uint32 BindingIndex;
-		if (!DescriptorInfo.GetDescriptorSetAndBindingIndex(FVulkanShaderHeader::UniformBuffer, BufferIndex, DescriptorSet, BindingIndex))
-		{
-			return;
-		}
+		DescriptorInfo.GetUBDescriptorSetAndBindingIndex(BufferIndex, DescriptorSet, BindingIndex);
 
 		const VkDescriptorType DescriptorType = DescriptorInfo.GetDescriptorType(DescriptorSet, BindingIndex);
 
diff --git a/Engine/Source/Runtime/VulkanRHI/Private/VulkanDescriptorSets.h b/Engine/Source/Runtime/VulkanRHI/Private/VulkanDescriptorSets.h
--- a/Engine/Source/Runtime/VulkanRHI/Private/VulkanDescriptorSets.h
+++ b/Engine/Source/Runtime/VulkanRHI/Private/VulkanDescriptorSets.h
@@ -674,29 +674,17 @@ public:
 		//		FMemory::Memzero(CodeHeaders);
 	}
 
-	inline bool GetDescriptorSetAndBindingIndex(const FVulkanShaderHeader::EType Type, int32 ParameterIndex, uint8& OutDescriptorSet, uint32& OutBindingIndex) const
+	inline void GetUBDescriptorSetAndBindingIndex(int32 ParameterIndex, uint8& OutDescriptorSet, uint32& OutBindingIndex) const
 	{
-		switch (Type)
-		{
-		case FVulkanShaderHeader::UniformBuffer:
-			//ensure(RemappingInfo->StageInfos[0].UniformBuffers[ParameterIndex].bHasConstantData);
-			ensure(RemappingUBInfos[ParameterIndex].bHasConstantData);
-			//OutDescriptorSet = RemappingInfo->StageInfos[0].UniformBuffers[ParameterIndex].Remapping.NewDescriptorSet;
-			//OutBindingIndex = RemappingInfo->StageInfos[0].UniformBuffers[ParameterIndex].Remapping.NewBindingIndex;
-			OutDescriptorSet = RemappingUBInfos[ParameterIndex].Remapping.NewDescriptorSet;
-			OutBindingIndex = RemappingUBInfos[ParameterIndex].Remapping.NewBindingIndex;
-			break;
-		case FVulkanShaderHeader::Global:
-			//OutDescriptorSet = RemappingInfo->StageInfos[0].Globals[ParameterIndex].NewDescriptorSet;
-			//OutBindingIndex = RemappingInfo->StageInfos[0].Globals[ParameterIndex].NewBindingIndex;
-			OutDescriptorSet = RemappingGlobalInfos[ParameterIndex].NewDescriptorSet;
-			OutBindingIndex = RemappingGlobalInfos[ParameterIndex].NewBindingIndex;
-			break;
-		default:
-			check(0);
-			return false;
-		}
-		return true;
+		ensure(RemappingUBInfos[ParameterIndex].bHasConstantData);
+		OutDescriptorSet = RemappingUBInfos[ParameterIndex].Remapping.NewDescriptorSet;
+		OutBindingIndex = RemappingUBInfos[ParameterIndex].Remapping.NewBindingIndex;
+	}
+
+	inline void GetDescriptorSetAndBindingIndex(int32 ParameterIndex, uint8& OutDescriptorSet, uint32& OutBindingIndex) const
+	{
+		OutDescriptorSet = RemappingGlobalInfos[ParameterIndex].NewDescriptorSet;
+		OutBindingIndex = RemappingGlobalInfos[ParameterIndex].NewBindingIndex;
 	}
 
 	inline const TArray<FDescriptorSetRemappingInfo::FRemappingInfo>& GetGlobalRemappingInfo() const
@@ -741,24 +729,17 @@ public:
 	{
 	}
 
-	inline bool GetDescriptorSetAndBindingIndex(const FVulkanShaderHeader::EType Type, const ShaderStage::EStage Stage, int32 ParameterIndex, uint8& OutDescriptorSet, uint32& OutBindingIndex) const
+	inline void GetUBDescriptorSetAndBindingIndex(const ShaderStage::EStage Stage, int32 ParameterIndex, uint8& OutDescriptorSet, uint32& OutBindingIndex) const
 	{
-		switch (Type)
-		{
-		case FVulkanShaderHeader::UniformBuffer:
-			ensure(RemappingUBInfos[Stage][ParameterIndex].bHasConstantData);
-			OutDescriptorSet = RemappingUBInfos[Stage][ParameterIndex].Remapping.NewDescriptorSet;
-			OutBindingIndex = RemappingUBInfos[Stage][ParameterIndex].Remapping.NewBindingIndex;
-			break;
-		case FVulkanShaderHeader::Global:
-			OutDescriptorSet = RemappingGlobalInfos[Stage][ParameterIndex].NewDescriptorSet;
-			OutBindingIndex = RemappingGlobalInfos[Stage][ParameterIndex].NewBindingIndex;
-			break;
-		default:
-			check(0);
-			return false;
-		}
-		return true;
+		ensure(RemappingUBInfos[Stage][ParameterIndex].bHasConstantData);
+		OutDescriptorSet = RemappingUBInfos[Stage][ParameterIndex].Remapping.NewDescriptorSet;
+		OutBindingIndex = RemappingUBInfos[Stage][ParameterIndex].Remapping.NewBindingIndex;
+	}
+
+	inline void GetDescriptorSetAndBindingIndex(const ShaderStage::EStage Stage, int32 ParameterIndex, uint8& OutDescriptorSet, uint32& OutBindingIndex) const
+	{
+		OutDescriptorSet = RemappingGlobalInfos[Stage][ParameterIndex].NewDescriptorSet;
+		OutBindingIndex = RemappingGlobalInfos[Stage][ParameterIndex].NewBindingIndex;
 	}
 
 	inline const TArray<FDescriptorSetRemappingInfo::FRemappingInfo>& GetGlobalRemappingInfo(ShaderStage::EStage Stage) const
diff --git a/Engine/Source/Runtime/VulkanRHI/Private/VulkanPendingState.h b/Engine/Source/Runtime/VulkanRHI/Private/VulkanPendingState.h
--- a/Engine/Source/Runtime/VulkanRHI/Private/VulkanPendingState.h
+++ b/Engine/Source/Runtime/VulkanRHI/Private/VulkanPendingState.h
@@ -69,11 +69,7 @@ public:
 		const FVulkanComputePipelineDescriptorInfo& DescriptorInfo = CurrentState->GetComputePipelineDescriptorInfo();
 		uint8 DescriptorSet;
 		uint32 BindingIndex;
-		if (!DescriptorInfo.GetDescriptorSetAndBindingIndex(FVulkanShaderHeader::Global, UAVIndex, DescriptorSet, BindingIndex))
-		{
-			return;
-		}
-
+		DescriptorInfo.GetDescriptorSetAndBindingIndex(UAVIndex, DescriptorSet, BindingIndex);
 		SetUAVForUBResource(DescriptorSet, BindingIndex, UAV);
 	}
 
@@ -82,11 +78,7 @@ public:
 		const FVulkanComputePipelineDescriptorInfo& DescriptorInfo = CurrentState->GetComputePipelineDescriptorInfo();
 		uint8 DescriptorSet;
 		uint32 BindingIndex;
-		if (!DescriptorInfo.GetDescriptorSetAndBindingIndex(FVulkanShaderHeader::Global, TextureIndex, DescriptorSet, BindingIndex))
-		{
-			return;
-		}
-
+		DescriptorInfo.GetDescriptorSetAndBindingIndex(TextureIndex, DescriptorSet, BindingIndex);
 		CurrentState->SetTexture(DescriptorSet, BindingIndex, Texture, Layout);
 	}
 
@@ -95,11 +87,7 @@ public:
 		const FVulkanComputePipelineDescriptorInfo& DescriptorInfo = CurrentState->GetComputePipelineDescriptorInfo();
 		uint8 DescriptorSet;
 		uint32 BindingIndex;
-		if (!DescriptorInfo.GetDescriptorSetAndBindingIndex(FVulkanShaderHeader::Global, SamplerIndex, DescriptorSet, BindingIndex))
-		{
-			return;
-		}
-
+		DescriptorInfo.GetDescriptorSetAndBindingIndex(SamplerIndex, DescriptorSet, BindingIndex);
 		CurrentState->SetSamplerState(DescriptorSet, BindingIndex, Sampler);
 	}
 
@@ -115,11 +103,7 @@ public:
 		const FVulkanComputePipelineDescriptorInfo& DescriptorInfo = CurrentState->GetComputePipelineDescriptorInfo();
 		uint8 DescriptorSet;
 		uint32 BindingIndex;
-		if (!DescriptorInfo.GetDescriptorSetAndBindingIndex(FVulkanShaderHeader::Global, SRVIndex, DescriptorSet, BindingIndex))
-		{
-			return;
-		}
-
+		DescriptorInfo.GetDescriptorSetAndBindingIndex(SRVIndex, DescriptorSet, BindingIndex);
 		SetSRVForUBResource(DescriptorSet, BindingIndex, SRV);
 	}
 
@@ -275,11 +259,7 @@ public:
 		const FVulkanGfxPipelineDescriptorInfo& DescriptorInfo = CurrentState->GetGfxPipelineDescriptorInfo();
 		uint8 DescriptorSet;
 		uint32 BindingIndex;
-		if (!DescriptorInfo.GetDescriptorSetAndBindingIndex(FVulkanShaderHeader::Global, Stage, ParameterIndex, DescriptorSet, BindingIndex))
-		{
-			return;
-		}
-
+		DescriptorInfo.GetDescriptorSetAndBindingIndex(Stage, ParameterIndex, DescriptorSet, BindingIndex);
 		CurrentState->SetTexture(DescriptorSet, BindingIndex, Texture, Layout);
 	}
 
@@ -301,11 +281,7 @@ public:
 		const FVulkanGfxPipelineDescriptorInfo& DescriptorInfo = CurrentState->GetGfxPipelineDescriptorInfo();
 		uint8 DescriptorSet;
 		uint32 BindingIndex;
-		if (!DescriptorInfo.GetDescriptorSetAndBindingIndex(FVulkanShaderHeader::Global, Stage, ParameterIndex, DescriptorSet, BindingIndex))
-		{
-			return;
-		}
-
+		DescriptorInfo.GetDescriptorSetAndBindingIndex(Stage, ParameterIndex, DescriptorSet, BindingIndex);
 		SetUAVForUBResource(DescriptorSet, BindingIndex, UAV);
 	}
 
@@ -316,11 +292,7 @@ public:
 		const FVulkanGfxPipelineDescriptorInfo& DescriptorInfo = CurrentState->GetGfxPipelineDescriptorInfo();
 		uint8 DescriptorSet;
 		uint32 BindingIndex;
-		if (!DescriptorInfo.GetDescriptorSetAndBindingIndex(FVulkanShaderHeader::Global, Stage, ParameterIndex, DescriptorSet, BindingIndex))
-		{
-			return;
-		}
-
+		DescriptorInfo.GetDescriptorSetAndBindingIndex(Stage, ParameterIndex, DescriptorSet, BindingIndex);
 		SetSRVForUBResource(DescriptorSet, BindingIndex, SRV);
 	}
 
@@ -329,11 +301,7 @@ public:
 		const FVulkanGfxPipelineDescriptorInfo& DescriptorInfo = CurrentState->GetGfxPipelineDescriptorInfo();
 		uint8 DescriptorSet;
 		uint32 BindingIndex;
-		if (!DescriptorInfo.GetDescriptorSetAndBindingIndex(FVulkanShaderHeader::Global, Stage, ParameterIndex, DescriptorSet, BindingIndex))
-		{
-			return;
-		}
-
+		DescriptorInfo.GetDescriptorSetAndBindingIndex(Stage, ParameterIndex, DescriptorSet, BindingIndex);
 		CurrentState->SetSamplerState(DescriptorSet, BindingIndex, Sampler);
 	}
 
diff --git a/Engine/Source/Runtime/VulkanRHI/Public/VulkanShaderResources.h b/Engine/Source/Runtime/VulkanRHI/Public/VulkanShaderResources.h
--- a/Engine/Source/Runtime/VulkanRHI/Public/VulkanShaderResources.h
+++ b/Engine/Source/Runtime/VulkanRHI/Public/VulkanShaderResources.h
@@ -18,15 +18,6 @@
 // Size = Ignored for non-globals
 struct FVulkanShaderHeader
 {
-	enum EType
-	{
-		PackedGlobal,
-		Global,
-		UniformBuffer,
-
-		Count,
-	};
-
 	struct FSpirvInfo
 	{
 		FSpirvInfo() = default;
@@ -40,8 +31,6 @@ struct FVulkanShaderHeader
 		uint32	BindingIndexOffset = UINT32_MAX;
 	};
 
-
-
 	struct FUniformBufferInfo
 	{
 		uint32					LayoutHash;
@@ -63,17 +52,6 @@ struct FVulkanShaderHeader
 	};
 	TArray<FGlobalInfo>					Globals;
 
-	struct FPackedGlobalInfo
-	{
-		uint16							ConstantDataSizeInFloats;
-		uint8							PackedUBIndex;
-		uint8							Pad0 = 0;
-#if VULKAN_ENABLE_BINDING_DEBUG_NAMES
-		FString							DebugName;
-#endif
-	};
-	TArray<FPackedGlobalInfo>			PackedGlobals;
-
 	struct FPackedUBInfo
 	{
 		uint32							SizeInBytes;
@@ -167,16 +145,6 @@ inline FArchive& operator<<(FArchive& Ar, FVulkanShaderHeader::FUniformBufferInf
 	return Ar;
 }
 
-inline FArchive& operator<<(FArchive& Ar, FVulkanShaderHeader::FPackedGlobalInfo& PackedGlobalInfo)
-{
-	Ar << PackedGlobalInfo.ConstantDataSizeInFloats;
-	Ar << PackedGlobalInfo.PackedUBIndex;
-#if VULKAN_ENABLE_BINDING_DEBUG_NAMES
-	Ar << PackedGlobalInfo.DebugName;
-#endif
-	return Ar;
-}
-
 inline FArchive& operator<<(FArchive& Ar, FVulkanShaderHeader::FPackedUBInfo& PackedUBInfo)
 {
 	Ar << PackedUBInfo.SizeInBytes;
@@ -205,7 +173,6 @@ inline FArchive& operator<<(FArchive& Ar, FVulkanShaderHeader& Header)
 {
 	Ar << Header.UniformBuffers;
 	Ar << Header.Globals;
-	Ar << Header.PackedGlobals;
 	Ar << Header.PackedUBs;
 	Ar << Header.InputAttachments;
 	Ar << Header.InOutMask;
