VulkanRHI: Remove remapping of bindings at runtime.  Locked packed globals to binding 0 of every stage if present (removed multiple loops).  Simplify Vulkan shader compiler by removing more legacy code paths.  Reorder EAttachmentType in FVulkanShaderHeader to match with VULKAN_SUBPASS_FETCH.

#rnx

diff --git a/Engine/Shaders/Public/Platform/Vulkan/VulkanCommon.ush b/Engine/Shaders/Public/Platform/Vulkan/VulkanCommon.ush
--- a/Engine/Shaders/Public/Platform/Vulkan/VulkanCommon.ush
+++ b/Engine/Shaders/Public/Platform/Vulkan/VulkanCommon.ush
@@ -7,7 +7,7 @@
 #pragma once
 
 // Update this GUID to invalidate shader recompilation for only Vulkan shaders
-#pragma message("UESHADERMETADATA_VERSION 05414BAE-522D-4FFC-98A8-BE9ADE764257")
+#pragma message("UESHADERMETADATA_VERSION 78B8BEEB-B86D-4C6D-8541-BB93F687F063")
 
 #define COMPILER_SUPPORTS_ULONG_TYPES 1
 #define COMPILER_SUPPORTS_UINT64_IMAGE_ATOMICS (1 && COMPILER_SUPPORTS_ULONG_TYPES)
diff --git a/Engine/Source/Developer/VulkanShaderFormat/Private/VulkanShaderCompiler.cpp b/Engine/Source/Developer/VulkanShaderFormat/Private/VulkanShaderCompiler.cpp
--- a/Engine/Source/Developer/VulkanShaderFormat/Private/VulkanShaderCompiler.cpp
+++ b/Engine/Source/Developer/VulkanShaderFormat/Private/VulkanShaderCompiler.cpp
@@ -165,26 +165,6 @@ uint32 ParseNumber(const T* Str, bool bEmptyIsZero = false)
 	return Num;
 }
 
-static void GetResourceEntryFromUBMember(const FShaderResourceTableMap& ResourceTableMap, const FString& UBName, uint16 ResourceIndex, FUniformResourceEntry& OutEntry)
-{
-	for (const FUniformResourceEntry& Entry : ResourceTableMap.Resources)
-	{
-		if (Entry.GetUniformBufferName() == UBName && Entry.ResourceIndex == ResourceIndex)
-		{
-			OutEntry = Entry;
-			return;
-		}
-	}
-
-	check(0);
-}
-
-
-struct FPatchType
-{
-	int32	HeaderGlobalIndex;
-	uint16	CombinedAliasIndex;
-};
 
 
 static const FString kBindlessCBPrefix = TEXT("__BindlessCB");
@@ -339,409 +319,31 @@ struct VulkanShaderCompilerSerializedOutput
 	FShaderResourceTable ShaderResourceTable;
 	FVulkanSpirv Spirv;
 
-	uint32 InputAttachmentsMask = 0;
-
 	TSet<FString> UsedBindlessUB;
 
 	FShaderCompilerResourceTable CompilerSRT;
 };
 
 
-static int32 AddGlobal( const FString& ParameterName,
-						uint16 BindingIndex,
-						VulkanShaderCompilerSerializedOutput& SerializedOutput, 
-						const TArray<FString>& GlobalNames
-)
-{
-	const int32 HeaderGlobalIndex = GlobalNames.Find(ParameterName);
-	check(HeaderGlobalIndex != INDEX_NONE);
-	check(GlobalNames[HeaderGlobalIndex] == ParameterName);
-
-	FVulkanShaderHeader::FGlobalInfo& GlobalInfo = SerializedOutput.Header.Globals[HeaderGlobalIndex];
-	const FVulkanSpirv::FEntry* Entry = SerializedOutput.Spirv.GetEntry(ParameterName);
-	if (Entry)
-	{
-		if (Entry->Binding == -1)
-		{
-			// Texel buffers get put into a uniform block
-			Entry = SerializedOutput.Spirv.GetEntry(ParameterName + TEXT("_BUFFER"));
-			check(Entry);
-			check(Entry->Binding != -1);
-		}
-	}
-	else
-	{
-		Entry = SerializedOutput.Spirv.GetEntryByBindingIndex(BindingIndex);
-		check(Entry);
-		check(Entry->Binding != -1);
-		if (!Entry->Name.EndsWith(TEXT("_BUFFER")))
-		{
-			checkf(false, TEXT("CombinedSamplers should not be used anymore"))
-		}
-	}
-
-	SerializedOutput.Header.GlobalSpirvInfos[HeaderGlobalIndex] = FVulkanShaderHeader::FSpirvInfo(Entry->WordDescriptorSetIndex, Entry->WordBindingIndex);
-
-	GlobalInfo.TypeIndex = Entry->DescriptorType;
-
-#if VULKAN_ENABLE_BINDING_DEBUG_NAMES
-	GlobalInfo.DebugName = ParameterName;
-#endif
-
-	return HeaderGlobalIndex;
-}
-
-static int32 AddUBResources(	const int32 HeaderUBIndex,
-							const FString& UBName,
-							const FShaderResourceTableMap& ResourceTableMap,
-							uint32 BufferIndex,
-							const TArray<uint32>& BindingArray,
-							VulkanShaderCompilerSerializedOutput& SerializedOutput,
-							TArray<FString>& GlobalNames)
-{
-	int32 NumAdded = 0;
-	if (BindingArray.Num() > 0)
-	{
-		uint32 BufferOffset = BindingArray[BufferIndex];
-		if (BufferOffset > 0)
-		{
-			// Extract all resources related to the current BufferIndex
-			const uint32* ResourceInfos = &BindingArray[BufferOffset];
-			uint32 ResourceInfo = *ResourceInfos++;
-			do
-			{
-				// Verify that we have correct buffer index
-				check(FRHIResourceTableEntry::GetUniformBufferIndex(ResourceInfo) == BufferIndex);
-
-				// Extract binding index from ResourceInfo
-				const uint32 BindingIndex = FRHIResourceTableEntry::GetBindIndex(ResourceInfo);
-
-				// Extract index of the resource stored in the resource table from ResourceInfo
-				const uint16 ResourceIndex = FRHIResourceTableEntry::GetResourceIndex(ResourceInfo);
 
-				FUniformResourceEntry ResourceTableEntry;
-				GetResourceEntryFromUBMember(ResourceTableMap, UBName, ResourceIndex, ResourceTableEntry);
-
-				const int32 HeaderGlobalIndex = AddGlobal(ResourceTableEntry.UniformBufferMemberName, BindingIndex, SerializedOutput, GlobalNames);
-
-				// Fill the compiler SRT with info about this resource
-				{
-					SerializedOutput.CompilerSRT.ResourceTableBits |= 1 << HeaderUBIndex;
-					const uint32 ResourceMap = FRHIResourceTableEntry::Create(HeaderUBIndex, ResourceIndex, HeaderGlobalIndex);
-					switch ((EUniformBufferBaseType)ResourceTableEntry.Type)
-					{
-					case UBMT_TEXTURE:
-					case UBMT_RDG_TEXTURE:
-						SerializedOutput.CompilerSRT.TextureMap.Add(ResourceMap);
-						break;
-					case UBMT_SAMPLER:
-						SerializedOutput.CompilerSRT.SamplerMap.Add(ResourceMap);
-						break;
-					case UBMT_SRV:
-					case UBMT_RDG_TEXTURE_SRV:
-					case UBMT_RDG_TEXTURE_NON_PIXEL_SRV:
-					case UBMT_RDG_BUFFER_SRV:
-						SerializedOutput.CompilerSRT.ShaderResourceViewMap.Add(ResourceMap);
-						break;
-					case UBMT_RESOURCE_COLLECTION:
-						SerializedOutput.CompilerSRT.ResourceCollectionMap.Add(ResourceMap);
-						break;
-					case UBMT_UAV:
-					case UBMT_RDG_TEXTURE_UAV:
-					case UBMT_RDG_BUFFER_UAV:
-						SerializedOutput.CompilerSRT.UnorderedAccessViewMap.Add(ResourceMap);
-						break;
-					default:
-						check(false);
-					}
-
-					SerializedOutput.CompilerSRT.MaxBoundResourceTable = FMath::Max<uint32>(HeaderUBIndex+1, SerializedOutput.CompilerSRT.MaxBoundResourceTable);
-
-					NumAdded++;
-				}
-
-				// Iterate to next info
-				ResourceInfo = *ResourceInfos++;
-			}
-			while (FRHIResourceTableEntry::GetUniformBufferIndex(ResourceInfo) == BufferIndex);
-		}
-	}
-	return NumAdded;
-}
-
-static void AddUniformBuffer(
-	const FShaderResourceTable& ShaderResourceTable,
-	const FShaderCompilerInput& ShaderInput,
-	const FString& UBName,
-	uint16 BindingIndex,
-	FShaderParameterMap& InOutParameterMap,
-	VulkanShaderCompilerSerializedOutput& SerializedOutput,
-	TArray<FString>& GlobalNames
-)
+static uint32 GetUBLayoutHash(const FShaderCompilerInput& ShaderInput, const FString& UBName)
 {
-	FVulkanShaderHeader& OutHeader = SerializedOutput.Header;
-
-	const int32 HeaderUBIndex = OutHeader.UniformBuffers.AddZeroed();
-	FVulkanShaderHeader::FUniformBufferInfo& UBInfo = OutHeader.UniformBuffers[HeaderUBIndex];
+	uint32 LayoutHash = 0;
 
 	const FUniformBufferEntry* UniformBufferEntry = ShaderInput.Environment.UniformBufferMap.Find(UBName);
 	if (UniformBufferEntry)
 	{
-		UBInfo.LayoutHash = UniformBufferEntry->LayoutHash;
+		LayoutHash = UniformBufferEntry->LayoutHash;
 	}
 	else if ((UBName == FShaderParametersMetadata::kRootUniformBufferBindingName) && ShaderInput.RootParametersStructure)
 	{
-		UBInfo.LayoutHash = ShaderInput.RootParametersStructure->GetLayoutHash();
-	}
-	else
-	{
-		UBInfo.LayoutHash = 0;
-	}
-	const int32 NewSRTHashIndex = SerializedOutput.CompilerSRT.ResourceTableLayoutHashes.Emplace(UBInfo.LayoutHash);
-	check(HeaderUBIndex == NewSRTHashIndex);
-
-#if VULKAN_ENABLE_BINDING_DEBUG_NAMES
-	UBInfo.DebugName = UBName;
-#endif
-
-	const FVulkanSpirv::FEntry* Entry = SerializedOutput.Spirv.GetEntry(UBName);
-	if (Entry)
-	{
-		UBInfo.bOnlyHasResources = false;
-		UBInfo.ConstantDataOriginalBindingIndex = BindingIndex;
-		
-		int32 SpirvInfoIndex = OutHeader.UniformBufferSpirvInfos.Add(FVulkanShaderHeader::FSpirvInfo(Entry->WordDescriptorSetIndex, Entry->WordBindingIndex));
-		check(SpirvInfoIndex == HeaderUBIndex);
-	}
-	else
-	{
-		UBInfo.bOnlyHasResources = true;
-		UBInfo.ConstantDataOriginalBindingIndex = UINT16_MAX;
-
-		int32 SpirvInfoIndex = OutHeader.UniformBufferSpirvInfos.Add(FVulkanShaderHeader::FSpirvInfo());
-		check(SpirvInfoIndex == HeaderUBIndex);
-	}
-
-	// Add used resources...
-	if (ShaderResourceTable.ResourceTableBits & (1 << BindingIndex))
-	{
-		int32 NumAdded = 0;
-
-		// Make sure to process in the same order as when gathering names below
-		NumAdded += AddUBResources(HeaderUBIndex, UBName, ShaderInput.Environment.ResourceTableMap, BindingIndex, ShaderResourceTable.TextureMap, SerializedOutput, GlobalNames);
-		NumAdded += AddUBResources(HeaderUBIndex, UBName, ShaderInput.Environment.ResourceTableMap, BindingIndex, ShaderResourceTable.SamplerMap, SerializedOutput, GlobalNames);
-		NumAdded += AddUBResources(HeaderUBIndex, UBName, ShaderInput.Environment.ResourceTableMap, BindingIndex, ShaderResourceTable.ShaderResourceViewMap, SerializedOutput, GlobalNames);
-		NumAdded += AddUBResources(HeaderUBIndex, UBName, ShaderInput.Environment.ResourceTableMap, BindingIndex, ShaderResourceTable.UnorderedAccessViewMap, SerializedOutput, GlobalNames);
-
-		UBInfo.bHasResources = (NumAdded > 0) ? 1 : 0;
-	}
-	else
-	{
-		// If we're using real uniform buffers we have to have resources at least
-		checkf(!UBInfo.bOnlyHasResources, TEXT("UBName = %s, BindingIndex = %d"), *UBName, (int32)BindingIndex);
-	}
-
-	InOutParameterMap.RemoveParameterAllocation(UBName);
-	InOutParameterMap.AddParameterAllocation(UBName, HeaderUBIndex, 0, 1, EShaderParameterType::UniformBuffer);
-}
-
-static int32 DoAddGlobal(const FString& Name, FVulkanShaderHeader& OutHeader, TArray<FString>& OutGlobalNames)
-{
-	check(!OutGlobalNames.Contains(Name));
-	const int32 NameIndex = OutGlobalNames.Add(Name);
-	const int32 GlobalIndex = OutHeader.Globals.AddDefaulted();
-	check(NameIndex == GlobalIndex);
-	const int32 GlobalSpirvIndex = OutHeader.GlobalSpirvInfos.AddDefaulted();
-	check(GlobalSpirvIndex == GlobalIndex);
-	return GlobalIndex;
-}
-
-static void PrepareUBResourceEntryGlobals(const TArray<uint32>& BindingArray, const FShaderResourceTableMap& ResourceTableMap,
-	int32 BufferIndex, const FString& UBName, TArray<FString>& OutGlobalNames, FVulkanShaderHeader& OutHeader)
-{
-	if (BindingArray.Num() > 0)
-	{
-		uint32 BufferOffset = BindingArray[BufferIndex];
-		if (BufferOffset > 0)
-		{
-			// Extract all resources related to the current BufferIndex
-			const uint32* ResourceInfos = &BindingArray[BufferOffset];
-			uint32 ResourceInfo = *ResourceInfos++;
-			do
-			{
-				// Verify that we have correct buffer index
-				check(FRHIResourceTableEntry::GetUniformBufferIndex(ResourceInfo) == BufferIndex);
-
-				// Extract index of the resource stored in the resource table from ResourceInfo
-				const uint16 ResourceIndex = FRHIResourceTableEntry::GetResourceIndex(ResourceInfo);
-
-				FUniformResourceEntry ResourceTableEntry;
-				GetResourceEntryFromUBMember(ResourceTableMap, UBName, ResourceIndex, ResourceTableEntry);
-
-				DoAddGlobal(ResourceTableEntry.UniformBufferMemberName, OutHeader, OutGlobalNames);
-
-				// Iterate to next info
-				ResourceInfo = *ResourceInfos++;
-			}
-			while (FRHIResourceTableEntry::GetUniformBufferIndex(ResourceInfo) == BufferIndex);
-		}
-	}
-}
-
-
-static void PrepareGlobals(const VulkanShaderCompilerSerializedOutput& SerializedOutput, const FSpirvReflectBindings& SpirvReflectBindings, const FShaderResourceTable& SRT, const FShaderCompilerInput& ShaderInput, const TArray<FString>& ParameterNames, FShaderParameterMap& ParameterMap, TArray<FString>& OutGlobalNames, FVulkanShaderHeader& OutHeader)
-{
-	// First pass, gather names for all the Globals that are NOT Samplers
-	for (int32 ParameterIndex = 0; ParameterIndex < ParameterNames.Num(); ++ParameterIndex)
-	{
-		const FString& ParameterName = ParameterNames[ParameterIndex];
-		TOptional<FParameterAllocation> ParameterAllocation = ParameterMap.FindParameterAllocation(ParameterName);
-		checkf(ParameterAllocation.IsSet(), TEXT("PrepareGlobals failed to find resource ParameterName=%s"), *ParameterName);
-
-		if (ParameterAllocation->Type == EShaderParameterType::UniformBuffer)
-		{
-			// Add used resources...
-			if (SRT.ResourceTableBits & (1 << ParameterAllocation->BufferIndex))
-			{
-				PrepareUBResourceEntryGlobals(SRT.TextureMap, ShaderInput.Environment.ResourceTableMap, ParameterAllocation->BufferIndex, ParameterName, OutGlobalNames, OutHeader);
-				PrepareUBResourceEntryGlobals(SRT.ShaderResourceViewMap, ShaderInput.Environment.ResourceTableMap, ParameterAllocation->BufferIndex, ParameterName, OutGlobalNames, OutHeader);
-				PrepareUBResourceEntryGlobals(SRT.UnorderedAccessViewMap, ShaderInput.Environment.ResourceTableMap, ParameterAllocation->BufferIndex, ParameterName, OutGlobalNames, OutHeader);
-			}
-		}
-		else if ((ParameterAllocation->Type == EShaderParameterType::LooseData) || IsParameterBindless(ParameterAllocation->Type))
-		{
-			// Do nothing
-		}
-		else if (ParameterAllocation->Type != EShaderParameterType::Sampler)
-		{
-			DoAddGlobal(ParameterName, OutHeader, OutGlobalNames);
-		}
+		LayoutHash = ShaderInput.RootParametersStructure->GetLayoutHash();
 	}
 
-	// Second pass, add all samplers
-	for (int32 ParameterIndex = 0; ParameterIndex < ParameterNames.Num(); ++ParameterIndex)
-	{
-		const FString& ParameterName = ParameterNames[ParameterIndex];
-		TOptional<FParameterAllocation> ParameterAllocation = ParameterMap.FindParameterAllocation(ParameterName);
-		checkf(ParameterAllocation.IsSet(), TEXT("PrepareGlobals failed to find resource ParameterName=%s"), *ParameterName);
-			
-		if (ParameterAllocation->Type == EShaderParameterType::UniformBuffer)
-		{
-			// Add used samplers...
-			if (SRT.ResourceTableBits & (1 << ParameterAllocation->BufferIndex))
-			{
-				PrepareUBResourceEntryGlobals(SRT.SamplerMap, ShaderInput.Environment.ResourceTableMap, ParameterAllocation->BufferIndex, ParameterName, OutGlobalNames, OutHeader);
-			}
-		}
-		else if (ParameterAllocation->Type == EShaderParameterType::Sampler)
-		{
-			DoAddGlobal(ParameterName, OutHeader, OutGlobalNames);
-		}
-	}
-
-	// Now input attachments
-	if (SerializedOutput.InputAttachmentsMask != 0)
-	{
-		uint32 InputAttachmentsMask = SerializedOutput.InputAttachmentsMask;
-		for (int32 Index = 0; InputAttachmentsMask != 0; ++Index, InputAttachmentsMask>>= 1)
-		{
-			if (InputAttachmentsMask & 1)
-			{
-				DoAddGlobal(VULKAN_SUBPASS_FETCH_VAR_W[Index], OutHeader, OutGlobalNames);
-			}
-		}
-	}
+	return LayoutHash;
 }
 
-static void ConvertToHeader(
-	FShaderResourceTable& ShaderResourceTable,
-	const FShaderCompilerInput& ShaderInput,
-	const FSpirvReflectBindings& SpirvReflectBindings,
-	FShaderParameterMap& InOutParameterMap,
-	VulkanShaderCompilerSerializedOutput& SerializedOutput
-)
-{
-	FVulkanShaderHeader& OutHeader = SerializedOutput.Header;
-
-	// Names that match the Header.Globals array
-	TArray<FString> GlobalNames;
-
-	TArray<FString> ParameterNames;
-	InOutParameterMap.GetAllParameterNames(ParameterNames);
-
-	PrepareGlobals(SerializedOutput, SpirvReflectBindings, ShaderResourceTable, ShaderInput, ParameterNames, InOutParameterMap, GlobalNames, OutHeader);
-
-	for (int32 ParameterIndex = 0; ParameterIndex < ParameterNames.Num(); ++ParameterIndex)
-	{
-		const FString& ParameterName = *ParameterNames[ParameterIndex];
-		const FParameterAllocation* ParameterAllocation = InOutParameterMap.GetParameterMap().Find(*ParameterName);
-		check(ParameterAllocation);
 
-		if (ParameterAllocation->Type == EShaderParameterType::UniformBuffer)
-		{
-			AddUniformBuffer(ShaderResourceTable, ShaderInput, ParameterName, ParameterAllocation->BufferIndex, InOutParameterMap, SerializedOutput, GlobalNames);
-		}
-		else if ((ParameterAllocation->Type == EShaderParameterType::LooseData) || IsParameterBindless(ParameterAllocation->Type))
-		{
-			// Do nothing
-		}
-		else
-		{
-			const int32 HeaderGlobalIndex = AddGlobal(ParameterName, ParameterAllocation->BaseIndex, SerializedOutput, GlobalNames);
-
-			InOutParameterMap.RemoveParameterAllocation(ParameterName);
-			InOutParameterMap.AddParameterAllocation(ParameterName, 0, HeaderGlobalIndex, ParameterAllocation->Size, ParameterAllocation->Type);
-		}
-	}
-
-	// Finally check for subpass/input attachments
-	if (SerializedOutput.InputAttachmentsMask != 0)
-	{
-		const static FVulkanShaderHeader::EAttachmentType AttachmentTypes[] = 
-		{
-			FVulkanShaderHeader::EAttachmentType::Depth,
-			FVulkanShaderHeader::EAttachmentType::Color0,
-			FVulkanShaderHeader::EAttachmentType::Color1,
-			FVulkanShaderHeader::EAttachmentType::Color2,
-			FVulkanShaderHeader::EAttachmentType::Color3,
-			FVulkanShaderHeader::EAttachmentType::Color4,
-			FVulkanShaderHeader::EAttachmentType::Color5,
-			FVulkanShaderHeader::EAttachmentType::Color6,
-			FVulkanShaderHeader::EAttachmentType::Color7
-		};
-
-		uint32 InputAttachmentsMask = SerializedOutput.InputAttachmentsMask;
-		for (int32 Index = 0; InputAttachmentsMask != 0; ++Index, InputAttachmentsMask>>=1)
-		{
-			if ((InputAttachmentsMask & 1) == 0)
-			{
-				continue;
-			}
-
-			const FString AttachmentName(VULKAN_SUBPASS_FETCH_VAR_W[Index]);
-			const VkDescriptorType DescriptorType = VK_DESCRIPTOR_TYPE_INPUT_ATTACHMENT;
-			const FVulkanShaderHeader::EAttachmentType AttachmentType = AttachmentTypes[Index];
-			{
-				const int32 HeaderGlobalIndex = GlobalNames.Find(AttachmentName);
-				check(HeaderGlobalIndex != INDEX_NONE);
-				check(GlobalNames[HeaderGlobalIndex] == AttachmentName);
-				FVulkanShaderHeader::FGlobalInfo& GlobalInfo = OutHeader.Globals[HeaderGlobalIndex];
-				const FVulkanSpirv::FEntry* Entry = SerializedOutput.Spirv.GetEntry(AttachmentName);
-				check(Entry);
-				check(Entry->Binding != -1);
-
-				OutHeader.GlobalSpirvInfos[HeaderGlobalIndex] = FVulkanShaderHeader::FSpirvInfo(Entry->WordDescriptorSetIndex, Entry->WordBindingIndex);
-				GlobalInfo.TypeIndex = (uint32)DescriptorType;
-#if VULKAN_ENABLE_BINDING_DEBUG_NAMES
-				GlobalInfo.DebugName = AttachmentName;
-#endif
-				FVulkanShaderHeader::FInputAttachment& AttachmentInfo = OutHeader.InputAttachments.AddZeroed_GetRef();
-				AttachmentInfo.GlobalIndex = HeaderGlobalIndex;
-				AttachmentInfo.Type = AttachmentType;
-			}
-		}
-	}
-}
 
 
 static void BuildShaderOutput(
@@ -750,24 +352,19 @@ static void BuildShaderOutput(
 	const FVulkanShaderCompilerInternalState& InternalState,
 	const FSpirvReflectBindings& SpirvReflectBindings,
 	const FString&				DebugName,
-	uint32						PackedGlobalArraySize,
 	TBitArray<>&				UsedUniformBufferSlots
 )
 {
 	const FShaderCompilerInput& ShaderInput = InternalState.Input;
 	const EShaderFrequency Frequency = InternalState.GetShaderFrequency();
+	FVulkanShaderHeader& Header = SerializedOutput.Header;
 
-	FVulkanShaderHeader& NEWHeader = SerializedOutput.Header;
-
-	NEWHeader.SpirvCRC = SerializedOutput.Spirv.CRC;
-	NEWHeader.RayTracingPayloadType = ShaderInput.Environment.GetCompileArgument(TEXT("RT_PAYLOAD_TYPE"), 0u);
-	NEWHeader.RayTracingPayloadSize = ShaderInput.Environment.GetCompileArgument(TEXT("RT_PAYLOAD_MAX_SIZE"), 0u);
-#if VULKAN_ENABLE_BINDING_DEBUG_NAMES
-	NEWHeader.DebugName = DebugName;
-#endif
+	Header.SpirvCRC = SerializedOutput.Spirv.CRC;
+	Header.RayTracingPayloadType = ShaderInput.Environment.GetCompileArgument(TEXT("RT_PAYLOAD_TYPE"), 0u);
+	Header.RayTracingPayloadSize = ShaderInput.Environment.GetCompileArgument(TEXT("RT_PAYLOAD_MAX_SIZE"), 0u);
 
 	// :todo-jn: Hash entire SPIRV for now, could eventually be removed since we use ShaderKeys
-	FSHA1::HashBuffer(SerializedOutput.Spirv.Data.GetData(), SerializedOutput.Spirv.GetByteSize(), (uint8*)&NEWHeader.SourceHash);
+	FSHA1::HashBuffer(SerializedOutput.Spirv.Data.GetData(), SerializedOutput.Spirv.GetByteSize(), (uint8*)&Header.SourceHash);
 
 
 	// Flattens the array dimensions of the interface variable (aka shader attribute), e.g. from float4[2][3] -> float4[6]
@@ -807,7 +404,7 @@ static void BuildShaderOutput(
 				for (uint32 Index = 0; Index < FlattenedArrayDim; ++Index)
 				{
 					const uint32 BitIndex = (AttributeIndex + Index);
-					NEWHeader.InOutMask |= (1u << BitIndex);
+					Header.InOutMask |= (1u << BitIndex);
 				}
 			}
 		}
@@ -824,7 +421,7 @@ static void BuildShaderOutput(
 			if (Attribute->built_in == SpvBuiltInFragDepth)
 			{
 				const uint32 BitIndex = (CrossCompiler::FShaderBindingInOutMask::DepthStencilMaskIndex);
-				NEWHeader.InOutMask |= (1u << BitIndex);
+				Header.InOutMask |= (1u << BitIndex);
 			}
 			else
 			{
@@ -838,70 +435,52 @@ static void BuildShaderOutput(
 					for (uint32 Index = 0; Index < FlattenedArrayDim; ++Index)
 					{
 						const uint32 BitIndex = (TargetIndex + Index);
-						NEWHeader.InOutMask |= (1u << BitIndex);
+						Header.InOutMask |= (1u << BitIndex);
 					}
 				}
 			}
 		}
 	}
 	
-	const int32 StageOffset = InternalState.bSupportsBindless ? (ShaderStage::GetStageForFrequency(Frequency) * VulkanBindless::MaxUniformBuffersPerStage) : 0;
-
-	TArray<FString> OriginalParameters;
-	ShaderOutput.ParameterMap.GetAllParameterNames(OriginalParameters);
-
 	// Build the SRT for this shader.
-	FShaderResourceTable ShaderResourceTable;
 	{
-		// Build the generic SRT for this shader.
-		FShaderCompilerResourceTable GenericSRT;
-		if (!BuildResourceTableMapping(ShaderInput.Environment.ResourceTableMap, ShaderInput.Environment.UniformBufferMap, UsedUniformBufferSlots, ShaderOutput.ParameterMap, /*MaxBoundResourceTable, */GenericSRT))
+		checkf(Header.UniformBufferInfos.Num() == (UsedUniformBufferSlots.FindLast(true) + 1), 
+			TEXT("Some of the Uniform Buffers containing constants weren't flag as in-use.  This might lead to duplicate indices being assigned."));
+
+		if (!BuildResourceTableMapping(ShaderInput.Environment.ResourceTableMap, ShaderInput.Environment.UniformBufferMap, UsedUniformBufferSlots, ShaderOutput.ParameterMap, SerializedOutput.CompilerSRT))
 		{
 			ShaderOutput.Errors.Add(TEXT("Internal error on BuildResourceTableMapping."));
 			return;
 		}
+		UE::ShaderCompilerCommon::BuildShaderResourceTable(SerializedOutput.CompilerSRT, SerializedOutput.ShaderResourceTable);
 
-		UE::ShaderCompilerCommon::BuildShaderResourceTable(GenericSRT, ShaderResourceTable, true);
-	}
-
-	TArray<FString> NewParameters;
-	ShaderOutput.ParameterMap.GetAllParameterNames(NewParameters);
-
-	// Mark all used uniform buffer indices; however some are empty (eg GBuffers) so gather those as NewParameters
-	uint16 NumParams = 0;
-	for (int32 Index = NewParameters.Num() - 1; Index >= 0; --Index)
-	{
-		uint16 OutIndex, OutBase, OutSize;
-		const bool bFound = ShaderOutput.ParameterMap.FindParameterAllocation(NewParameters[Index], OutIndex, OutBase, OutSize);
-		ensure(bFound);
-		NumParams = FMath::Max((uint16)(OutIndex + 1), NumParams);
-		if (OriginalParameters.Contains(NewParameters[Index]))
+		// The previous step also added resource only UBs starting at the first free slot in UsedUniformBufferSlots
+		// We need to add the hashes for their layouts in the same slots of our UniformBufferInfos in the header
 		{
-			NewParameters.RemoveAtSwap(Index, EAllowShrinking::No);
-		}
-	}
-
-	if (PackedGlobalArraySize > 0)
-	{
-		FVulkanShaderHeader::FPackedUBInfo& PackedUB = NEWHeader.PackedUBs.AddZeroed_GetRef();
-		PackedUB.SizeInBytes = Align(PackedGlobalArraySize, 16u);
+			const int32 NumUBSlots = SerializedOutput.CompilerSRT.MaxBoundResourceTable + 1;
+			if (Header.UniformBufferInfos.Num() < NumUBSlots)
+			{
+				Header.UniformBufferInfos.SetNumZeroed(NumUBSlots);
+			}
 
-		const FVulkanSpirv::FEntry* Entry = SerializedOutput.Spirv.GetEntryByBindingIndex(StageOffset);
-		check(Entry);
-		PackedUB.SPIRVDescriptorSetOffset = Entry->WordDescriptorSetIndex;
-		PackedUB.SPIRVBindingIndexOffset = Entry->WordBindingIndex;
-	}
+			TArray<FStringView> UBParameterNames = ShaderOutput.ParameterMap.GetAllParameterNamesOfType(EShaderParameterType::UniformBuffer);
+			for (const FStringView& ParameterName : UBParameterNames)
+			{
+				TOptional<FParameterAllocation> Allocation = ShaderOutput.ParameterMap.FindParameterAllocation(ParameterName);
+				check(Allocation.IsSet());
+				const uint32 UniformBufferIndex = Allocation.GetValue().BufferIndex;
 
-	ConvertToHeader(ShaderResourceTable, ShaderInput, SpirvReflectBindings, ShaderOutput.ParameterMap, SerializedOutput);
+				FVulkanShaderHeader::FUniformBufferInfo& UniformBufferInfo = Header.UniformBufferInfos[UniformBufferIndex];
+				UniformBufferInfo.bHasResources = 1;
 
-	if (ShaderInput.Environment.CompilerFlags.Contains(CFLAG_ExtraShaderData))
-	{
-		NEWHeader.DebugName = ShaderInput.GenerateShaderName();
+				const uint32 UBLayoutHash = SerializedOutput.CompilerSRT.ResourceTableLayoutHashes[UniformBufferIndex];
+				checkf(!UniformBufferInfo.LayoutHash || (UniformBufferInfo.LayoutHash == UBLayoutHash), 
+					TEXT("Layout hash should be unset (resource only UB) or identical (UB with constants and resources)!"));
+				UniformBufferInfo.LayoutHash = UBLayoutHash;
+			}
+		}
 	}
 
-	// Build the SRT for this shader from accumulated data in the Compiler SRT
-	UE::ShaderCompilerCommon::BuildShaderResourceTable(SerializedOutput.CompilerSRT, SerializedOutput.ShaderResourceTable);
-
 	ShaderOutput.bSucceeded = true;
 
 	// guard disassembly of SPIRV code on bExtractShaderSource setting since presumably this isn't that cheap.
@@ -929,6 +508,15 @@ static void BuildShaderOutput(
 	{
 		CullGlobalUniformBuffers(ShaderInput.Environment.UniformBufferMap, ShaderOutput.ParameterMap);
 	}
+
+#if VULKAN_ENABLE_BINDING_DEBUG_NAMES
+	Header.DebugName = DebugName;
+#else
+	if (ShaderInput.Environment.CompilerFlags.Contains(CFLAG_ExtraShaderData))
+	{
+		Header.DebugName = ShaderInput.GenerateShaderName();
+	}
+#endif
 }
 
 
@@ -1110,8 +698,6 @@ static bool BuildShaderOutputFromSpirv(
 	const FString UBOGlobalsNameSpv(ANSI_TO_TCHAR(CrossCompiler::FShaderConductorContext::GetIdentifierTable().GlobalsUniformBuffer));
 	const FString UBORootParamNameSpv(FShaderParametersMetadata::kRootUniformBufferBindingName);
 
-	uint32& InputAttachmentsMask = SerializedOutput.InputAttachmentsMask;
-	uint32 PackedGlobalArraySize = 0;
 	TBitArray<> UsedUniformBufferSlots;
 	const int32 MaxNumBits = VulkanBindless::MaxUniformBuffersPerStage * SF_NumFrequencies;
 	UsedUniformBufferSlots.Init(false, MaxNumBits);
@@ -1122,9 +708,22 @@ static bool BuildShaderOutputFromSpirv(
 		const int32 StageOffset = InternalState.bSupportsBindless ? (UEStage * VulkanBindless::MaxUniformBuffersPerStage) : 0;
 		const uint32_t DescSetNumber = InternalState.bSupportsBindless ? (uint32_t)VulkanBindless::BindlessSingleUseUniformBufferSet : (uint32_t)UEStage;
 
-		int32 GlobalBindingIndex = StageOffset;
+		auto AddShaderValidationType = [](uint32_t VulkanBindingIndex, const FShaderParameterParser::FParsedShaderParameter* ParsedParam, FShaderCompilerOutput& Output) 
+		{
+			/*if (ParsedParam)
+			{
+				if (IsResourceBindingTypeSRV(ParsedParam->ParsedTypeDecl))
+				{
+					AddShaderValidationSRVType(VulkanBindingIndex, ParsedParam->ParsedTypeDecl, Output);
+				}
+				else
+				{
+					AddShaderValidationUAVType(VulkanBindingIndex, ParsedParam->ParsedTypeDecl, Output);
+				}
+			}*/
+		};
 
-		auto AddReflectionInfos = [&](TArray<SpvReflectDescriptorBinding*>& BindingArray, const VkDescriptorType DescriptorType, int32 BindingOffset, bool bIsPackedUniformBuffer=false)
+		auto AddReflectionInfos = [&](TArray<SpvReflectDescriptorBinding*>& BindingArray, const VkDescriptorType DescriptorType, int32 BindingTypeCount, bool bIsPackedUniformBuffer=false)
 		{
 			for (const SpvReflectDescriptorBinding* Binding : BindingArray)
 			{
@@ -1139,111 +738,109 @@ static bool BuildShaderOutputFromSpirv(
 					continue;
 				}
 
-				const int32 BindingIndex = GlobalBindingIndex++;
+				const int32 BindingSlot = SerializedOutput.Header.Bindings.Num();
+				const int32 BindingIndex = StageOffset + BindingSlot;
+				FVulkanShaderHeader::FBindingInfo& BindingInfo = SerializedOutput.Header.Bindings.AddZeroed_GetRef();
+				BindingInfo.DescriptorType = DescriptorType;
+
+#if VULKAN_ENABLE_BINDING_DEBUG_NAMES
+				BindingInfo.DebugName = ResourceName;
+#endif
 
 				const SpvReflectResult SpvResult = Reflection.ChangeDescriptorBindingNumbers(Binding, BindingIndex, DescSetNumber);
 				check(SpvResult == SPV_REFLECT_RESULT_SUCCESS);
 
 				const int32 ReflectionSlot = SerializedOutput.Spirv.ReflectionInfo.Add(FVulkanSpirv::FEntry(ResourceName, BindingIndex, (uint32)DescriptorType));
+				check(BindingSlot == ReflectionSlot);
 				check(InternalState.ParameterParser);
 				const FShaderParameterParser::FParsedShaderParameter* ParsedParam = InternalState.ParameterParser->FindParameterInfosUnsafe(ResourceName);
 
-				auto AddShaderValidationType = [] (uint32_t VulkanBindingIndex, const FShaderParameterParser::FParsedShaderParameter* ParsedParam, FShaderCompilerOutput& Output) {
-					/*if (ParsedParam)
-					{
-						if (IsResourceBindingTypeSRV(ParsedParam->ParsedTypeDecl))
-						{
-							AddShaderValidationSRVType(VulkanBindingIndex, ParsedParam->ParsedTypeDecl, Output);
-						}
-						else
-						{
-							AddShaderValidationUAVType(VulkanBindingIndex, ParsedParam->ParsedTypeDecl, Output);
-						}
-					}*/
-				};
-
 				switch (DescriptorType)
 				{
 				case VK_DESCRIPTOR_TYPE_STORAGE_TEXEL_BUFFER:
 				case VK_DESCRIPTOR_TYPE_STORAGE_BUFFER:
 				case VK_DESCRIPTOR_TYPE_STORAGE_IMAGE:
-					HandleReflectedShaderUAV(ResourceName, BindingOffset, ReflectionSlot, 1, Output);
-					AddShaderValidationType(BindingOffset, ParsedParam, Output);
+					HandleReflectedShaderUAV(ResourceName, BindingTypeCount, ReflectionSlot, 1, Output);
+					AddShaderValidationType(BindingTypeCount, ParsedParam, Output);
 					break;
 
 				case VK_DESCRIPTOR_TYPE_SAMPLED_IMAGE:
-					// todo-jn: Could verify that we have the samplers...
-					//for (uint32 UsageIndex = 0; UsageIndex < Binding->usage_binding_count; ++UsageIndex)
-					//{
-					//	const SpvReflectDescriptorBinding* AssociatedResource = Binding->usage_bindings[UsageIndex];
-					//	AssociatedResourceNames[UsageIndex] = ANSI_TO_TCHAR(AssociatedResource->name);
-					//}
-					[[fallthrough]];
-
 				case VK_DESCRIPTOR_TYPE_UNIFORM_TEXEL_BUFFER:
-					HandleReflectedShaderResource(ResourceName, BindingOffset, ReflectionSlot, 1, Output);
-					AddShaderValidationType(BindingOffset, ParsedParam, Output);
+					HandleReflectedShaderResource(ResourceName, BindingTypeCount, ReflectionSlot, 1, Output);
+					AddShaderValidationType(BindingTypeCount, ParsedParam, Output);
 					break;
 
 				case VK_DESCRIPTOR_TYPE_SAMPLER:
 					HandleReflectedShaderSampler(ResourceName, ReflectionSlot, Output);
-					//HandleReflectedShaderSampler(ResourceName, BindingOffset, ReflectionSlot, 1, Output);
+					//HandleReflectedShaderSampler(ResourceName, BindingTypeCount, ReflectionSlot, 1, Output);
 					break;
 
 				case VK_DESCRIPTOR_TYPE_ACCELERATION_STRUCTURE_KHR:
-					HandleReflectedShaderResource(ResourceName, BindingOffset, ReflectionSlot, 1, Output);
-					AddShaderValidationType(BindingOffset, ParsedParam, Output);
+					HandleReflectedShaderResource(ResourceName, BindingTypeCount, ReflectionSlot, 1, Output);
+					AddShaderValidationType(BindingTypeCount, ParsedParam, Output);
 					break;
 
 				case VK_DESCRIPTOR_TYPE_INPUT_ATTACHMENT:
-					InputAttachmentsMask |= (1u << Binding->input_attachment_index);
+					{
+						SerializedOutput.Header.InputAttachmentsMask |= (1u << Binding->input_attachment_index);
+						FVulkanShaderHeader::FInputAttachmentInfo& InputAttachmentInfo = SerializedOutput.Header.InputAttachmentInfos.AddZeroed_GetRef();
+						InputAttachmentInfo.BindingIndex = BindingSlot;
+						InputAttachmentInfo.Type = (FVulkanShaderHeader::EAttachmentType)Binding->input_attachment_index;
+					}
 					break;
 
 				case VK_DESCRIPTOR_TYPE_UNIFORM_BUFFER:
-					if (bIsPackedUniformBuffer)
 					{
-						// Use the given global ResourceName instead of patching it to _Globals_h
-						check(!UsedUniformBufferSlots[ReflectionSlot]);
-						UsedUniformBufferSlots[ReflectionSlot] = true;
-
-						if (InternalState.UseRootParametersStructure())
-						{
-							check(ReflectionSlot == FShaderParametersMetadata::kRootCBufferBindingIndex);
-							HandleReflectedUniformBuffer(ResourceName, ReflectionSlot, Output);
-						}
-
-						// Register all uniform buffer members of Globals as loose data
-						for (uint32 MemberIndex = 0; MemberIndex < Binding->block.member_count; ++MemberIndex)
+						if (bIsPackedUniformBuffer)
 						{
-							const SpvReflectBlockVariable& Member = Binding->block.members[MemberIndex];
-
-							FString MemberName(ANSI_TO_TCHAR(Member.name));
-							FStringView AdjustedMemberName(MemberName);
-
-							const EShaderParameterType BindlessParameterType = FShaderParameterParser::ParseAndRemoveBindlessParameterPrefix(AdjustedMemberName);
-
-							// Add all members of global ub, and only bindless samplers/resources for root param
-							if (!InternalState.UseRootParametersStructure() || BindlessParameterType != EShaderParameterType::LooseData)
+							// Use the given global ResourceName instead of patching it to _Globals_h
+							if (InternalState.UseRootParametersStructure())
 							{
-								check(BindingOffset == 0); // Global constants should always be the first UB
-								HandleReflectedGlobalConstantBufferMember(
-									MemberName,
-									BindingOffset,
-									Member.absolute_offset,
-									Member.size,
-									Output
-								);
+								check(ReflectionSlot == FShaderParametersMetadata::kRootCBufferBindingIndex);
+								HandleReflectedUniformBuffer(ResourceName, ReflectionSlot, Output);
 							}
 
-							PackedGlobalArraySize = FMath::Max<uint32>((Member.absolute_offset + Member.size), PackedGlobalArraySize);
+							// Register all uniform buffer members of Globals as loose data
+							for (uint32 MemberIndex = 0; MemberIndex < Binding->block.member_count; ++MemberIndex)
+							{
+								const SpvReflectBlockVariable& Member = Binding->block.members[MemberIndex];
+
+								FString MemberName(ANSI_TO_TCHAR(Member.name));
+								FStringView AdjustedMemberName(MemberName);
+
+								const EShaderParameterType BindlessParameterType = FShaderParameterParser::ParseAndRemoveBindlessParameterPrefix(AdjustedMemberName);
+
+								// Add all members of global ub, and only bindless samplers/resources for root param
+								if (!InternalState.UseRootParametersStructure() || BindlessParameterType != EShaderParameterType::LooseData)
+								{
+									check(BindingTypeCount == 0); // Global constants should always be the first UB
+									HandleReflectedGlobalConstantBufferMember(
+										MemberName,
+										BindingTypeCount,
+										Member.absolute_offset,
+										Member.size,
+										Output
+									);
+								}
+
+								SerializedOutput.Header.PackedGlobalsSize = FMath::Max<uint32>((Member.absolute_offset + Member.size), SerializedOutput.Header.PackedGlobalsSize);
+								SerializedOutput.Header.PackedGlobalsSize = Align(SerializedOutput.Header.PackedGlobalsSize, 16u);
+							}
 						}
-					}
-					else
-					{
+						else
+						{
+							check(BindingTypeCount == ReflectionSlot);
+							check(!UsedUniformBufferSlots[ReflectionSlot]);
+							HandleReflectedUniformBuffer(ResourceName, ReflectionSlot, Output);
+							AddShaderValidationUBSize(BindingTypeCount, Binding->block.padded_size, Output);
+						}
+
 						check(!UsedUniformBufferSlots[ReflectionSlot]);
 						UsedUniformBufferSlots[ReflectionSlot] = true;
-						HandleReflectedUniformBuffer(ResourceName, ReflectionSlot, Output);
-						AddShaderValidationUBSize(BindingOffset, Binding->block.padded_size, Output);
+
+						FVulkanShaderHeader::FUniformBufferInfo& UniformBufferInfo = SerializedOutput.Header.UniformBufferInfos.AddZeroed_GetRef();
+						UniformBufferInfo.LayoutHash = GetUBLayoutHash(InternalState.Input, ResourceName);
+						check(SerializedOutput.Header.Bindings.Num() == SerializedOutput.Header.UniformBufferInfos.Num());
 					}
 					break;
 
@@ -1252,14 +849,22 @@ static bool BuildShaderOutputFromSpirv(
 					break;
 				};
 
-				BindingOffset++;
+				// For every resource, 
+				if (DescriptorType != VK_DESCRIPTOR_TYPE_UNIFORM_BUFFER)
+				{
+
+				}
+
+				BindingTypeCount++;
 			}
-			return BindingOffset;
+
+			return BindingTypeCount;
 		};
 
 		// Process Globals first (PackedUniformBuffer) and then regular UBs
 		const int32 GlobalUBCount = AddReflectionInfos(Bindings.UniformBuffers, VK_DESCRIPTOR_TYPE_UNIFORM_BUFFER, 0, true);
-		int32 UBOBindings = AddReflectionInfos(Bindings.UniformBuffers, VK_DESCRIPTOR_TYPE_UNIFORM_BUFFER, 0) + GlobalUBCount;
+		const int32 UBOBindings = AddReflectionInfos(Bindings.UniformBuffers, VK_DESCRIPTOR_TYPE_UNIFORM_BUFFER, GlobalUBCount);
+		SerializedOutput.Header.NumBoundUniformBuffers = UBOBindings;
 
 		AddReflectionInfos(Bindings.InputAttachments, VK_DESCRIPTOR_TYPE_INPUT_ATTACHMENT, 0);
 
@@ -1329,7 +934,6 @@ static bool BuildShaderOutputFromSpirv(
 		InternalState,
 		Bindings,
 		InternalState.GetDebugName(),
-		PackedGlobalArraySize,
 		UsedUniformBufferSlots
 	);
 
@@ -1939,26 +1543,7 @@ static TArray<FString> ConvertUBToBindless(FString& PreprocessedShaderSource)
 
 static void UpdateBindlessUBs(const FVulkanShaderCompilerInternalState& InternalState, VulkanShaderCompilerSerializedOutput& SerializedOutput, FShaderCompilerOutput& Output)
 {
-	auto GetLayoutHash = [&InternalState](const FString& UBName)
-	{
-		uint32 LayoutHash = 0;
-		const FUniformBufferEntry* UniformBufferEntry = InternalState.Input.Environment.UniformBufferMap.Find(UBName);
-		if (UniformBufferEntry)
-		{
-			LayoutHash = UniformBufferEntry->LayoutHash;
-		}
-		else if ((UBName == FShaderParametersMetadata::kRootUniformBufferBindingName) && InternalState.Input.RootParametersStructure)
-		{
-			LayoutHash = InternalState.Input.RootParametersStructure->GetLayoutHash();
-		}
-		else
-		{
-			LayoutHash = 0;
-		}
-		return LayoutHash;
-	};
-
-	SerializedOutput.Header.UniformBuffers.Empty();
+	check(SerializedOutput.Header.Bindings.Num() == 0);
 	for (int32 CBIndex = 0; CBIndex < InternalState.AllBindlessUBs.Num(); CBIndex++)
 	{
 		const FString& CBName = InternalState.AllBindlessUBs[CBIndex];
@@ -1966,14 +1551,11 @@ static void UpdateBindlessUBs(const FVulkanShaderCompilerInternalState& Internal
 		// It's possible SPIRV compilation has optimized out a buffer from every shader in the group
 		if (SerializedOutput.UsedBindlessUB.Contains(CBName))
 		{
-			FVulkanShaderHeader::FUniformBufferInfo& UBInfo = SerializedOutput.Header.UniformBuffers.AddZeroed_GetRef();
-			UBInfo.LayoutHash = GetLayoutHash(CBName);
-			UBInfo.ConstantDataOriginalBindingIndex = CBIndex;
-#if VULKAN_ENABLE_BINDING_DEBUG_NAMES
-			UBInfo.DebugName = CBName;
-#endif
+			FVulkanShaderHeader::FUniformBufferInfo& Info = SerializedOutput.Header.UniformBufferInfos.AddZeroed_GetRef();
+			Info.LayoutHash = GetUBLayoutHash(InternalState.Input, CBName);
+			Info.BindlessCBIndex = CBIndex;
 
-			const int32 UBIndex = SerializedOutput.Header.UniformBuffers.Num() - 1;
+			const int32 UBIndex = SerializedOutput.Header.UniformBufferInfos.Num() - 1;
 			Output.ParameterMap.AddParameterAllocation(CBName, UBIndex, 0, 1, EShaderParameterType::UniformBuffer);
 		}
 	}
diff --git a/Engine/Source/Runtime/VulkanRHI/Private/VulkanCommands.cpp b/Engine/Source/Runtime/VulkanRHI/Private/VulkanCommands.cpp
--- a/Engine/Source/Runtime/VulkanRHI/Private/VulkanCommands.cpp
+++ b/Engine/Source/Runtime/VulkanRHI/Private/VulkanCommands.cpp
@@ -125,43 +125,43 @@ struct FVulkanResourceBinder
 {
 	FVulkanCommandListContext& Context;
 	const EShaderFrequency Frequency;
-	const TArray<FDescriptorSetRemappingInfo::FRemappingInfo>& GlobalRemappingInfo;
+	const ShaderStage::EStage Stage;
 	PendingStateType* PendingState;
 
-	FVulkanResourceBinder(FVulkanCommandListContext& InContext, EShaderFrequency InFrequency, const TArray<FDescriptorSetRemappingInfo::FRemappingInfo>& InGlobalRemappingInfo, PendingStateType* InPendingState)
+	FVulkanResourceBinder(FVulkanCommandListContext& InContext, EShaderFrequency InFrequency, PendingStateType* InPendingState)
 		: Context(InContext)
 		, Frequency(InFrequency)
-		, GlobalRemappingInfo(InGlobalRemappingInfo)
+		, Stage((InFrequency == SF_Compute) ? ShaderStage::Compute : ShaderStage::GetStageForFrequency(InFrequency))
 		, PendingState(InPendingState)
 	{
 	}
 
-	void SetUAV(FRHIUnorderedAccessView* UAV, uint8 Index, bool bClearResources = false)
+	void SetUAV(FRHIUnorderedAccessView* UAV, uint16 Index, bool bClearResources = false)
 	{
 		if (bClearResources)
 		{
 			//Context.ClearShaderResources(UAV);
 		}
 
-		PendingState->SetUAVForUBResource(GlobalRemappingInfo[Index].NewDescriptorSet, GlobalRemappingInfo[Index].NewBindingIndex, ResourceCast(UAV));
+		PendingState->SetUAVForUBResource(Stage, Index, ResourceCast(UAV));
 	}
 
-	void SetSRV(FRHIShaderResourceView* SRV, uint8 Index)
+	void SetSRV(FRHIShaderResourceView* SRV, uint16 Index)
 	{
-		PendingState->SetSRVForUBResource(GlobalRemappingInfo[Index].NewDescriptorSet, GlobalRemappingInfo[Index].NewBindingIndex, ResourceCast(SRV));
+		PendingState->SetSRVForUBResource(Stage, Index, ResourceCast(SRV));
 	}
 
-	void SetTexture(FRHITexture* TextureRHI, uint8 Index)
+	void SetTexture(FRHITexture* TextureRHI, uint16 Index)
 	{
 		FVulkanTexture* VulkanTexture = ResourceCast(TextureRHI);
 		const ERHIAccess RHIAccess = (Frequency == SF_Compute) ? ERHIAccess::SRVCompute : ERHIAccess::SRVGraphics;
 		const VkImageLayout ExpectedLayout = FVulkanLayoutManager::GetDefaultLayout(Context.GetCommandBufferManager()->GetActiveCmdBuffer(), *VulkanTexture, RHIAccess);
-		PendingState->SetTextureForUBResource(GlobalRemappingInfo[Index].NewDescriptorSet, GlobalRemappingInfo[Index].NewBindingIndex, VulkanTexture, ExpectedLayout);
+		PendingState->SetTextureForUBResource(Stage, Index, VulkanTexture, ExpectedLayout);
 	}
 
-	void SetSampler(FRHISamplerState* Sampler, uint8 Index)
+	void SetSampler(FRHISamplerState* Sampler, uint16 Index)
 	{
-		PendingState->SetSamplerStateForUBResource(GlobalRemappingInfo[Index].NewDescriptorSet, GlobalRemappingInfo[Index].NewBindingIndex, ResourceCast(Sampler));
+		PendingState->SetSamplerStateForUBResource(Stage, Index, ResourceCast(Sampler));
 	}
 
 #if PLATFORM_SUPPORTS_BINDLESS_RENDERING
@@ -182,10 +182,7 @@ void FVulkanCommandListContext::SetResourcesFromTables(const ShaderType* Shader)
 
 	if (Frequency == SF_Compute)
 	{
-		const TArray<FDescriptorSetRemappingInfo::FRemappingInfo>& GlobalRemappingInfo =
-			PendingComputeState->CurrentState->GetComputePipelineDescriptorInfo().GetGlobalRemappingInfo();
-
-		FVulkanResourceBinder Binder(*this, Frequency, GlobalRemappingInfo, PendingComputeState);
+		FVulkanResourceBinder Binder(*this, Frequency, PendingComputeState);
 		UE::RHICore::SetResourcesFromTables(
 			Binder
 			, *Shader
@@ -198,10 +195,7 @@ void FVulkanCommandListContext::SetResourcesFromTables(const ShaderType* Shader)
 	}
 	else
 	{
-		const TArray<FDescriptorSetRemappingInfo::FRemappingInfo>& GlobalRemappingInfo =
-		PendingGfxState->CurrentState->GetGfxPipelineDescriptorInfo().GetGlobalRemappingInfo(ShaderStage::GetStageForFrequency(Frequency));
-
-		FVulkanResourceBinder Binder(*this, Frequency, GlobalRemappingInfo, PendingGfxState);
+		FVulkanResourceBinder Binder(*this, Frequency, PendingGfxState);
 		UE::RHICore::SetResourcesFromTables(
 			Binder
 			, *Shader
@@ -485,15 +479,11 @@ void FVulkanCommandListContext::RHISetUniformBufferDynamicOffset(FUniformBufferS
 			const FUniformBufferStaticSlot Slot = StaticSlots[BufferIndex];
 			if (Slot == InSlot)
 			{
-				uint8 DescriptorSet;
-				uint32 BindingIndex;
-				DescriptorInfo.GetUBDescriptorSetAndBindingIndex(Stage, BufferIndex, DescriptorSet, BindingIndex);
-				{
-					// Uniform views always bind max supported range, so make sure Offset+Range is within buffer allocation
-					check((InOffset + PLATFORM_MAX_UNIFORM_BUFFER_RANGE) <= UniformBuffer->Allocation.Size);
-					uint32 DynamicOffset = InOffset + UniformBuffer->GetOffset();
-					PendingGfxState->CurrentState->SetUniformBufferDynamicOffset(DescriptorSet, BindingIndex, DynamicOffset);
-				}
+				// Uniform views always bind max supported range, so make sure Offset+Range is within buffer allocation
+				check((InOffset + PLATFORM_MAX_UNIFORM_BUFFER_RANGE) <= UniformBuffer->Allocation.Size);
+				uint32 DynamicOffset = InOffset + UniformBuffer->GetOffset();
+				PendingGfxState->CurrentState->SetUniformBufferDynamicOffset(Stage, BufferIndex, DynamicOffset);
+
 				break;
 			}
 		}
@@ -512,43 +502,42 @@ void FVulkanCommandListContext::RHISetShaderUniformBuffer(FRHIGraphicsShader* Sh
 
 	FVulkanUniformBuffer* UniformBuffer = ResourceCast(BufferRHI);
 	const FVulkanShaderHeader& CodeHeader = Shader->GetCodeHeader();
-	const FVulkanShaderHeader::FUniformBufferInfo& HeaderUBInfo = CodeHeader.UniformBuffers[BufferIndex];
-	checkfSlow(!HeaderUBInfo.LayoutHash || HeaderUBInfo.LayoutHash == UniformBuffer->GetLayout().GetHash(), TEXT("Mismatched UB layout! Got hash 0x%x, expected 0x%x!"), UniformBuffer->GetLayout().GetHash(), HeaderUBInfo.LayoutHash);
+	checkfSlow(!CodeHeader.UniformBufferInfos[BufferIndex].LayoutHash || (CodeHeader.UniformBufferInfos[BufferIndex].LayoutHash == UniformBuffer->GetLayout().GetHash()),
+		TEXT("Mismatched UB layout! Got hash 0x%x, expected 0x%x!"), UniformBuffer->GetLayout().GetHash(), CodeHeader.UniformBufferInfos[BufferIndex].LayoutHash);
 	const FVulkanGfxPipelineDescriptorInfo& DescriptorInfo = PendingGfxState->CurrentState->GetGfxPipelineDescriptorInfo();
 
-	if (!HeaderUBInfo.bOnlyHasResources)
+	bool bHasResources = false;
+	if (BufferIndex < CodeHeader.NumBoundUniformBuffers)
 	{
 		checkSlow(UniformBuffer->GetLayout().ConstantBufferSize > 0);
 
-		uint8 DescriptorSet;
-		uint32 BindingIndex;
-		DescriptorInfo.GetUBDescriptorSetAndBindingIndex(Stage, BufferIndex, DescriptorSet, BindingIndex);
-
-		const VkDescriptorType DescriptorType = DescriptorInfo.GetDescriptorType(DescriptorSet, BindingIndex);
+		const VkDescriptorType DescriptorType = DescriptorInfo.GetDescriptorType(Stage, BufferIndex);
 
 		if (DescriptorType == VK_DESCRIPTOR_TYPE_UNIFORM_BUFFER_DYNAMIC)
 		{
-			PendingGfxState->SetUniformBuffer<true>(DescriptorSet, BindingIndex, UniformBuffer);
+			PendingGfxState->SetUniformBuffer<true>(Stage, BufferIndex, UniformBuffer);
 		}
 		else
 		{
 			check(DescriptorType == VK_DESCRIPTOR_TYPE_UNIFORM_BUFFER);
-			PendingGfxState->SetUniformBuffer<false>(DescriptorSet, BindingIndex, UniformBuffer);
+			PendingGfxState->SetUniformBuffer<false>(Stage, BufferIndex, UniformBuffer);
 		}
+
+		bHasResources = (CodeHeader.UniformBufferInfos[BufferIndex].bHasResources != 0);
+	}
+	else
+	{
+		// If the buffer has no bindings, then it is as resource only ub
+		bHasResources = true;
 	}
 
-	if (HeaderUBInfo.bHasResources)
+	if (bHasResources)
 	{
 		checkSlow(Shader->Frequency < SF_NumStandardFrequencies);
 		check(BufferIndex < MAX_UNIFORM_BUFFERS_PER_SHADER_STAGE);
 		BoundUniformBuffers[Shader->Frequency][BufferIndex] = UniformBuffer;
 		DirtyUniformBuffers[Shader->Frequency] |= (1 << BufferIndex);
 	}
-	else
-	{
-		// Internal error: Completely empty UB!
-		checkSlow(!HeaderUBInfo.bOnlyHasResources);
-	}
 }
 
 void FVulkanCommandListContext::RHISetShaderUniformBuffer(FRHIComputeShader* ComputeShaderRHI, uint32 BufferIndex, FRHIUniformBuffer* BufferRHI)
@@ -567,43 +556,42 @@ void FVulkanCommandListContext::RHISetShaderUniformBuffer(FRHIComputeShader* Com
 
 	const FVulkanComputePipelineDescriptorInfo& DescriptorInfo = PendingComputeState->CurrentState->GetComputePipelineDescriptorInfo();
 	const FVulkanShaderHeader& CodeHeader = Shader->GetCodeHeader();
-	const FVulkanShaderHeader::FUniformBufferInfo& HeaderUBInfo = CodeHeader.UniformBuffers[BufferIndex];
-	checkfSlow(!HeaderUBInfo.LayoutHash || HeaderUBInfo.LayoutHash == UniformBuffer->GetLayout().GetHash(), TEXT("Mismatched UB layout! Got hash 0x%x, expected 0x%x!"), UniformBuffer->GetLayout().GetHash(), HeaderUBInfo.LayoutHash);
+	checkfSlow(!CodeHeader.UniformBufferInfos[BufferIndex].LayoutHash || (CodeHeader.UniformBufferInfos[BufferIndex].LayoutHash == UniformBuffer->GetLayout().GetHash()), 
+		TEXT("Mismatched UB layout! Got hash 0x%x, expected 0x%x!"), UniformBuffer->GetLayout().GetHash(), CodeHeader.UniformBufferInfos[BufferIndex].LayoutHash);
 
 	// Uniform Buffers
-	if (!HeaderUBInfo.bOnlyHasResources)
+	bool bHasResources = false;
+	if (BufferIndex < CodeHeader.NumBoundUniformBuffers)
 	{
 		checkSlow(UniformBuffer->GetLayout().ConstantBufferSize > 0);
 		
-		uint8 DescriptorSet;
-		uint32 BindingIndex;
-		DescriptorInfo.GetUBDescriptorSetAndBindingIndex(BufferIndex, DescriptorSet, BindingIndex);
-
-		const VkDescriptorType DescriptorType = DescriptorInfo.GetDescriptorType(DescriptorSet, BindingIndex);
+		const VkDescriptorType DescriptorType = DescriptorInfo.GetDescriptorType(ShaderStage::Compute, BufferIndex);
 
 		if (DescriptorType == VK_DESCRIPTOR_TYPE_UNIFORM_BUFFER_DYNAMIC)
 		{
-			State.SetUniformBuffer<true>(DescriptorSet, BindingIndex, UniformBuffer);
+			State.SetUniformBuffer<true>(ShaderStage::Compute, BufferIndex, UniformBuffer);
 		}
 		else
 		{
 			check(DescriptorType == VK_DESCRIPTOR_TYPE_UNIFORM_BUFFER);
-			State.SetUniformBuffer<false>(DescriptorSet, BindingIndex, UniformBuffer);
+			State.SetUniformBuffer<false>(ShaderStage::Compute, BufferIndex, UniformBuffer);
 		}
+
+		bHasResources = (CodeHeader.UniformBufferInfos[BufferIndex].bHasResources != 0);
+	}
+	else
+	{
+		// If the buffer has no bindings, then it is as resource only ub
+		bHasResources = true;
 	}
 
-	if (HeaderUBInfo.bHasResources)
+	if (bHasResources)
 	{
 		checkSlow(ComputeShaderRHI->GetFrequency() == SF_Compute);
 		check(BufferIndex < MAX_UNIFORM_BUFFERS_PER_SHADER_STAGE);
 		BoundUniformBuffers[SF_Compute][BufferIndex] = UniformBuffer;
 		DirtyUniformBuffers[SF_Compute] |= (1 << BufferIndex);
 	}
-	else
-	{
-		// Internal error: Completely empty UB!
-		checkSlow(!HeaderUBInfo.bOnlyHasResources);
-	}
 }
 
 void FVulkanCommandListContext::RHISetStencilRef(uint32 StencilRef)
diff --git a/Engine/Source/Runtime/VulkanRHI/Private/VulkanContext.h b/Engine/Source/Runtime/VulkanRHI/Private/VulkanContext.h
--- a/Engine/Source/Runtime/VulkanRHI/Private/VulkanContext.h
+++ b/Engine/Source/Runtime/VulkanRHI/Private/VulkanContext.h
@@ -19,7 +19,6 @@ class FVulkanPendingComputeState;
 class FVulkanQueue;
 class FVulkanSwapChain;
 
-struct FInputAttachmentData;
 
 class FVulkanCommandListContext : public IRHICommandContext
 {
diff --git a/Engine/Source/Runtime/VulkanRHI/Private/VulkanDescriptorSets.h b/Engine/Source/Runtime/VulkanRHI/Private/VulkanDescriptorSets.h
--- a/Engine/Source/Runtime/VulkanRHI/Private/VulkanDescriptorSets.h
+++ b/Engine/Source/Runtime/VulkanRHI/Private/VulkanDescriptorSets.h
@@ -36,55 +36,32 @@ struct FUniformBufferGatherInfo
 	const FVulkanShaderHeader*	CodeHeaders[ShaderStage::NumStages];
 };
 
-struct FInputAttachmentData
-{
-	uint16	BindingIndex = UINT16_MAX;
-	uint8	DescriptorSet = UINT8_MAX;
-	FVulkanShaderHeader::EAttachmentType Type = FVulkanShaderHeader::EAttachmentType::Count;
-};
 
 // Information for remapping descriptor sets when combining layouts
 struct FDescriptorSetRemappingInfo
 {
-	struct FRemappingInfo
-	{
-		uint16	NewDescriptorSet = UINT16_MAX;
-		uint16	NewBindingIndex = UINT16_MAX;
-	};
-
-	struct FUBRemappingInfo
-	{
-		// Remapping is only valid if there is constant data
-		FRemappingInfo			Remapping;
-
-		bool					bHasConstantData = false;
-		const bool				bPadding = false;	// padding is need on memcmp/MemCrc to make sure mem align
-	};
-
 	struct FStageInfo
 	{
-		TArray<FRemappingInfo>		Globals;
-		TArray<FUBRemappingInfo>	UniformBuffers;
 		TArray<VkDescriptorType>	Types;
-		TArray<uint16>				PackedUBBindingIndices;
-		uint16						PackedUBDescriptorSet = UINT16_MAX;
+		uint32						PackedGlobalsSize = 0;
+		uint32						NumBoundUniformBuffers = 0;
 		uint16						NumImageInfos = 0;
 		uint16						NumBufferInfos = 0;
 		uint16						NumAccelerationStructures = 0;
 
 		inline bool IsEmpty() const
 		{
-			if (Globals.Num() != 0)
+			if (Types.Num() != 0)
 			{
 				return false;
 			}
 
-			if (UniformBuffers.Num() != 0)
+			if (PackedGlobalsSize != 0)
 			{
 				return false;
 			}
 
-			if (PackedUBBindingIndices.Num() != 0)
+			if (NumBoundUniformBuffers != 0)
 			{
 				return false;
 			}
@@ -94,7 +71,7 @@ struct FDescriptorSetRemappingInfo
 	};
 	TStaticArray<FStageInfo, ShaderStage::NumStages>	StageInfos;
 
-	TArray<FInputAttachmentData>						InputAttachmentData;
+	TArray<FVulkanShaderHeader::FInputAttachmentInfo>	InputAttachmentData;
 
 	inline bool operator==(const FDescriptorSetRemappingInfo& In) const
 	{
@@ -103,25 +80,20 @@ struct FDescriptorSetRemappingInfo
 			return false;
 		}
 
-		if (FMemory::Memcmp(InputAttachmentData.GetData(), In.InputAttachmentData.GetData(), sizeof(FInputAttachmentData) * InputAttachmentData.Num()))
+		if (FMemory::Memcmp(InputAttachmentData.GetData(), In.InputAttachmentData.GetData(), sizeof(FVulkanShaderHeader::FInputAttachmentInfo) * InputAttachmentData.Num()))
 		{
 			return false;
 		}
 
 		for (uint32 StageInfosIndex = 0; StageInfosIndex < ShaderStage::NumStages; ++StageInfosIndex)
 		{
-			if (StageInfos[StageInfosIndex].PackedUBDescriptorSet != In.StageInfos[StageInfosIndex].PackedUBDescriptorSet ||
+			if (StageInfos[StageInfosIndex].PackedGlobalsSize != In.StageInfos[StageInfosIndex].PackedGlobalsSize ||
+				StageInfos[StageInfosIndex].NumBoundUniformBuffers != In.StageInfos[StageInfosIndex].NumBoundUniformBuffers ||
 				StageInfos[StageInfosIndex].NumBufferInfos != In.StageInfos[StageInfosIndex].NumBufferInfos ||
 				StageInfos[StageInfosIndex].NumImageInfos != In.StageInfos[StageInfosIndex].NumImageInfos ||
 				StageInfos[StageInfosIndex].NumAccelerationStructures != In.StageInfos[StageInfosIndex].NumAccelerationStructures ||
-				StageInfos[StageInfosIndex].Globals.Num() != In.StageInfos[StageInfosIndex].Globals.Num() ||
-				StageInfos[StageInfosIndex].PackedUBBindingIndices.Num() != In.StageInfos[StageInfosIndex].PackedUBBindingIndices.Num() ||
-				StageInfos[StageInfosIndex].UniformBuffers.Num() != In.StageInfos[StageInfosIndex].UniformBuffers.Num() ||
 				StageInfos[StageInfosIndex].Types.Num() != In.StageInfos[StageInfosIndex].Types.Num() ||
-				FMemory::Memcmp(StageInfos[StageInfosIndex].Types.GetData(), In.StageInfos[StageInfosIndex].Types.GetData(), sizeof(VkDescriptorType) * StageInfos[StageInfosIndex].Types.Num()) ||
-				FMemory::Memcmp(StageInfos[StageInfosIndex].Globals.GetData(), In.StageInfos[StageInfosIndex].Globals.GetData(), sizeof(FRemappingInfo) * StageInfos[StageInfosIndex].Globals.Num()) ||
-				FMemory::Memcmp(StageInfos[StageInfosIndex].PackedUBBindingIndices.GetData(), In.StageInfos[StageInfosIndex].PackedUBBindingIndices.GetData(), sizeof(uint16) * StageInfos[StageInfosIndex].PackedUBBindingIndices.Num()) ||
-				FMemory::Memcmp(StageInfos[StageInfosIndex].UniformBuffers.GetData(), In.StageInfos[StageInfosIndex].UniformBuffers.GetData(), sizeof(FUBRemappingInfo) * StageInfos[StageInfosIndex].UniformBuffers.Num()))
+				FMemory::Memcmp(StageInfos[StageInfosIndex].Types.GetData(), In.StageInfos[StageInfosIndex].Types.GetData(), sizeof(VkDescriptorType) * StageInfos[StageInfosIndex].Types.Num()))
 			{
 				return false;
 			}
@@ -153,79 +125,6 @@ struct FDescriptorSetRemappingInfo
 		}
 		return true;
 	}
-
-	uint32 AddGlobal(uint32 Stage, int32 GlobalIndex, uint32 NewDescriptorSet, VkDescriptorType InType)
-	{
-		FStageInfo& StageInfo = StageInfos[Stage];
-
-		const uint32 NewBindingIndex = StageInfo.Types.Add(InType);
-
-		const int32 RemappingIndex = StageInfo.Globals.AddDefaulted();
-		check(RemappingIndex == GlobalIndex);
-		FDescriptorSetRemappingInfo::FRemappingInfo& Remapping = StageInfo.Globals[RemappingIndex];
-		Remapping.NewDescriptorSet = NewDescriptorSet;
-		Remapping.NewBindingIndex = NewBindingIndex;
-
-		return NewBindingIndex;
-	}
-
-	uint32 AddPackedUB(uint32 Stage, int32 PackUBIndex, uint32 NewDescriptorSet, VkDescriptorType InType)
-	{
-		FStageInfo& StageInfo = StageInfos[Stage];
-
-		const uint32 NewBindingIndex = StageInfo.Types.Add(InType);
-		if (StageInfo.PackedUBDescriptorSet == UINT16_MAX)
-		{
-			StageInfo.PackedUBDescriptorSet = NewDescriptorSet;
-		}
-		else
-		{
-			ensure(StageInfo.PackedUBDescriptorSet == NewDescriptorSet);
-		}
-		const int32 RemappingIndex = StageInfo.PackedUBBindingIndices.Add(NewBindingIndex);
-		check(RemappingIndex == PackUBIndex);
-
-		return NewBindingIndex;
-	}
-
-	void AddUBWithData(uint32 Stage, int32 UniformBufferIndex, uint32 NewDescriptorSet, VkDescriptorType InType, uint32& OutNewBindingIndex)
-	{
-		FStageInfo& StageInfo = StageInfos[Stage];
-
-		OutNewBindingIndex = StageInfo.Types.Add(InType);
-
-		const int32 UBRemappingIndex = StageInfo.UniformBuffers.AddDefaulted();
-		check(UBRemappingIndex == UniformBufferIndex);
-		FDescriptorSetRemappingInfo::FUBRemappingInfo& UBRemapping = StageInfo.UniformBuffers[UBRemappingIndex];
-		UBRemapping.bHasConstantData = true;
-		UBRemapping.Remapping.NewDescriptorSet = NewDescriptorSet;
-		UBRemapping.Remapping.NewBindingIndex = OutNewBindingIndex;
-	}
-
-	void AddRedundantUB(uint32 Stage, int32 UniformBufferIndex, const FDescriptorSetRemappingInfo::FUBRemappingInfo* InExistingUBInfo)
-	{
-		const int32 UBRemappingIndex = StageInfos[Stage].UniformBuffers.AddDefaulted();
-		check(UBRemappingIndex == UniformBufferIndex);
-		FDescriptorSetRemappingInfo::FUBRemappingInfo& UBRemapping = StageInfos[Stage].UniformBuffers[UBRemappingIndex];
-		UBRemapping.bHasConstantData = InExistingUBInfo->bHasConstantData;
-		UBRemapping.Remapping = InExistingUBInfo->Remapping;
-	}
-
-	void AddUBResourceOnly(uint32 Stage, int32 UniformBufferIndex)
-	{
-		const int32 UBRemappingIndex = StageInfos[Stage].UniformBuffers.AddDefaulted();
-		check(UBRemappingIndex == UniformBufferIndex);
-		FDescriptorSetRemappingInfo::FUBRemappingInfo& UBRemapping = StageInfos[Stage].UniformBuffers[UBRemappingIndex];
-		UBRemapping.bHasConstantData = false;
-	}
-
-	static_assert(
-		sizeof(FRemappingInfo) == (sizeof(FRemappingInfo::NewDescriptorSet) + sizeof(FRemappingInfo::NewBindingIndex)),
-		"FRemappingInfo should not have padding! structure is used for MemCrc32/Memcmp");
-
-	static_assert(
-		sizeof(FUBRemappingInfo) == (sizeof(FRemappingInfo) + sizeof(FUBRemappingInfo::bHasConstantData) + sizeof(FUBRemappingInfo::bPadding)),
-		"FUBRemappingInfo should not have padding! structure is used for MemCrc32/Memcmp");
 };
 
 // Information for the layout of descriptor sets; does not hold runtime objects
@@ -671,26 +570,6 @@ public:
 		, RemappingInfo(nullptr)
 		, bInitialized(false)
 	{
-		//		FMemory::Memzero(CodeHeaders);
-	}
-
-	inline void GetUBDescriptorSetAndBindingIndex(int32 ParameterIndex, uint8& OutDescriptorSet, uint32& OutBindingIndex) const
-	{
-		ensure(RemappingUBInfos[ParameterIndex].bHasConstantData);
-		OutDescriptorSet = RemappingUBInfos[ParameterIndex].Remapping.NewDescriptorSet;
-		OutBindingIndex = RemappingUBInfos[ParameterIndex].Remapping.NewBindingIndex;
-	}
-
-	inline void GetDescriptorSetAndBindingIndex(int32 ParameterIndex, uint8& OutDescriptorSet, uint32& OutBindingIndex) const
-	{
-		OutDescriptorSet = RemappingGlobalInfos[ParameterIndex].NewDescriptorSet;
-		OutBindingIndex = RemappingGlobalInfos[ParameterIndex].NewBindingIndex;
-	}
-
-	inline const TArray<FDescriptorSetRemappingInfo::FRemappingInfo>& GetGlobalRemappingInfo() const
-	{
-		//OutDescriptorSet = RemappingUBInfos[Stage][ParameterIndex].Remapping.NewDescriptorSet;
-		return RemappingInfo->StageInfos[ShaderStage::Compute].Globals;
 	}
 	
 	inline VkDescriptorType GetDescriptorType(uint8 DescriptorSet, int32 DescriptorIndex) const
@@ -708,9 +587,6 @@ public:
 
 protected:
 	// Cached data from FDescriptorSetRemappingInfo
-	TArrayView<const FDescriptorSetRemappingInfo::FUBRemappingInfo>	RemappingUBInfos;
-	TArrayView<const FDescriptorSetRemappingInfo::FRemappingInfo>	RemappingGlobalInfos;
-	TArrayView<const uint16>										RemappingPackedUBInfos;
 	uint32															HasDescriptorsInSetMask;
 	const FDescriptorSetRemappingInfo*								RemappingInfo;
 	bool															bInitialized;
@@ -729,24 +605,6 @@ public:
 	{
 	}
 
-	inline void GetUBDescriptorSetAndBindingIndex(const ShaderStage::EStage Stage, int32 ParameterIndex, uint8& OutDescriptorSet, uint32& OutBindingIndex) const
-	{
-		ensure(RemappingUBInfos[Stage][ParameterIndex].bHasConstantData);
-		OutDescriptorSet = RemappingUBInfos[Stage][ParameterIndex].Remapping.NewDescriptorSet;
-		OutBindingIndex = RemappingUBInfos[Stage][ParameterIndex].Remapping.NewBindingIndex;
-	}
-
-	inline void GetDescriptorSetAndBindingIndex(const ShaderStage::EStage Stage, int32 ParameterIndex, uint8& OutDescriptorSet, uint32& OutBindingIndex) const
-	{
-		OutDescriptorSet = RemappingGlobalInfos[Stage][ParameterIndex].NewDescriptorSet;
-		OutBindingIndex = RemappingGlobalInfos[Stage][ParameterIndex].NewBindingIndex;
-	}
-
-	inline const TArray<FDescriptorSetRemappingInfo::FRemappingInfo>& GetGlobalRemappingInfo(ShaderStage::EStage Stage) const
-	{
-		return RemappingInfo->StageInfos[Stage].Globals;
-	}
-
 	inline VkDescriptorType GetDescriptorType(uint8 DescriptorSet, int32 DescriptorIndex) const
 	{
 		return RemappingInfo->StageInfos[DescriptorSet].Types[DescriptorIndex];
@@ -757,7 +615,7 @@ public:
 		return bInitialized;
 	}
 
-	inline const TArray<FInputAttachmentData>& GetInputAttachmentData() const
+	inline const TArray<FVulkanShaderHeader::FInputAttachmentInfo>& GetInputAttachmentData() const
 	{
 		return RemappingInfo->InputAttachmentData;
 	}
@@ -766,9 +624,6 @@ public:
 
 protected:
 	// Cached data from FDescriptorSetRemappingInfo
-	TArrayView<const FDescriptorSetRemappingInfo::FUBRemappingInfo>	RemappingUBInfos[ShaderStage::NumStages];
-	TArrayView<const FDescriptorSetRemappingInfo::FRemappingInfo>	RemappingGlobalInfos[ShaderStage::NumStages];
-	TArrayView<const uint16>										RemappingPackedUBInfos[ShaderStage::NumStages];
 	uint32													HasDescriptorsInSetMask;
 
 	const FDescriptorSetRemappingInfo*						RemappingInfo;
@@ -1227,8 +1082,6 @@ public:
 		return DescriptorSetLayout.GetHash();
 	}
 
-	void PatchSpirvBindings(FVulkanShader::FSpirvCode& SpirvCode, EShaderFrequency Frequency, const FVulkanShaderHeader& CodeHeader) const;
-
 protected:
 	FVulkanDescriptorSetsLayout	DescriptorSetLayout;
 	VkPipelineLayout			PipelineLayout;
diff --git a/Engine/Source/Runtime/VulkanRHI/Private/VulkanPendingState.cpp b/Engine/Source/Runtime/VulkanRHI/Private/VulkanPendingState.cpp
--- a/Engine/Source/Runtime/VulkanRHI/Private/VulkanPendingState.cpp
+++ b/Engine/Source/Runtime/VulkanRHI/Private/VulkanPendingState.cpp
@@ -534,13 +534,12 @@ void FVulkanPendingGfxState::InternalUpdateDynamicStates(FVulkanCmdBuffer* Cmd)
 void FVulkanPendingGfxState::UpdateInputAttachments(FVulkanFramebuffer* Framebuffer)
 {
 	const FVulkanGfxPipelineDescriptorInfo& GfxDescriptorInfo = CurrentState->GetGfxPipelineDescriptorInfo();
-	const TArray<FInputAttachmentData>& InputAttachmentData = GfxDescriptorInfo.GetInputAttachmentData();
+	const TArray<FVulkanShaderHeader::FInputAttachmentInfo>& InputAttachmentData = GfxDescriptorInfo.GetInputAttachmentData();
 
 	for (int32 Index = 0; Index < InputAttachmentData.Num(); ++Index)
 	{
-		const FInputAttachmentData& AttachmentData = InputAttachmentData[Index];
-		const uint32 ColorIndex = static_cast<uint32>(AttachmentData.Type);
-		
+		const FVulkanShaderHeader::FInputAttachmentInfo& AttachmentData = InputAttachmentData[Index];
+	
 		switch (AttachmentData.Type)
 		{
 		case FVulkanShaderHeader::EAttachmentType::Color0:
@@ -551,11 +550,14 @@ void FVulkanPendingGfxState::UpdateInputAttachments(FVulkanFramebuffer* Framebuf
 		case FVulkanShaderHeader::EAttachmentType::Color5:
 		case FVulkanShaderHeader::EAttachmentType::Color6:
 		case FVulkanShaderHeader::EAttachmentType::Color7:
-			check(ColorIndex < Framebuffer->GetNumColorAttachments());
-			CurrentState->SetInputAttachment(AttachmentData.DescriptorSet, AttachmentData.BindingIndex, Framebuffer->AttachmentTextureViews[ColorIndex]->GetTextureView(), VK_IMAGE_LAYOUT_SHADER_READ_ONLY_OPTIMAL);
+			{
+				const int32 ColorIndex = static_cast<int32>(AttachmentData.Type) - (int32)FVulkanShaderHeader::EAttachmentType::Color0;
+				check((ColorIndex >= 0) && (ColorIndex < (int32)Framebuffer->GetNumColorAttachments()));
+				CurrentState->SetInputAttachment(ShaderStage::Pixel, AttachmentData.BindingIndex, Framebuffer->AttachmentTextureViews[ColorIndex]->GetTextureView(), VK_IMAGE_LAYOUT_SHADER_READ_ONLY_OPTIMAL);
+			}
 			break;
 		case FVulkanShaderHeader::EAttachmentType::Depth:
-			CurrentState->SetInputAttachment(AttachmentData.DescriptorSet, AttachmentData.BindingIndex, Framebuffer->GetPartialDepthTextureView(), VK_IMAGE_LAYOUT_DEPTH_STENCIL_READ_ONLY_OPTIMAL);
+			CurrentState->SetInputAttachment(ShaderStage::Pixel, AttachmentData.BindingIndex, Framebuffer->GetPartialDepthTextureView(), VK_IMAGE_LAYOUT_DEPTH_STENCIL_READ_ONLY_OPTIMAL);
 			break;
 		default:
 			check(0);
diff --git a/Engine/Source/Runtime/VulkanRHI/Private/VulkanPendingState.h b/Engine/Source/Runtime/VulkanRHI/Private/VulkanPendingState.h
--- a/Engine/Source/Runtime/VulkanRHI/Private/VulkanPendingState.h
+++ b/Engine/Source/Runtime/VulkanRHI/Private/VulkanPendingState.h
@@ -66,29 +66,17 @@ public:
 
 	inline void SetUAVForStage(uint32 UAVIndex, FVulkanUnorderedAccessView* UAV)
 	{
-		const FVulkanComputePipelineDescriptorInfo& DescriptorInfo = CurrentState->GetComputePipelineDescriptorInfo();
-		uint8 DescriptorSet;
-		uint32 BindingIndex;
-		DescriptorInfo.GetDescriptorSetAndBindingIndex(UAVIndex, DescriptorSet, BindingIndex);
-		SetUAVForUBResource(DescriptorSet, BindingIndex, UAV);
+		SetUAVForUBResource(ShaderStage::Compute, UAVIndex, UAV);
 	}
 
 	inline void SetTextureForStage(uint32 TextureIndex, const FVulkanTexture* Texture, VkImageLayout Layout)
 	{
-		const FVulkanComputePipelineDescriptorInfo& DescriptorInfo = CurrentState->GetComputePipelineDescriptorInfo();
-		uint8 DescriptorSet;
-		uint32 BindingIndex;
-		DescriptorInfo.GetDescriptorSetAndBindingIndex(TextureIndex, DescriptorSet, BindingIndex);
-		CurrentState->SetTexture(DescriptorSet, BindingIndex, Texture, Layout);
+		CurrentState->SetTexture(ShaderStage::Compute, TextureIndex, Texture, Layout);
 	}
 
 	inline void SetSamplerStateForStage(uint32 SamplerIndex, FVulkanSamplerState* Sampler)
 	{
-		const FVulkanComputePipelineDescriptorInfo& DescriptorInfo = CurrentState->GetComputePipelineDescriptorInfo();
-		uint8 DescriptorSet;
-		uint32 BindingIndex;
-		DescriptorInfo.GetDescriptorSetAndBindingIndex(SamplerIndex, DescriptorSet, BindingIndex);
-		CurrentState->SetSamplerState(DescriptorSet, BindingIndex, Sampler);
+		CurrentState->SetSamplerState(ShaderStage::Compute, SamplerIndex, Sampler);
 	}
 
 	inline void SetTextureForUBResource(int32 DescriptorSet, uint32 BindingIndex, const FVulkanTexture* Texture, VkImageLayout Layout)
@@ -100,11 +88,7 @@ public:
 
 	inline void SetSRVForStage(uint32 SRVIndex, FVulkanShaderResourceView* SRV)
 	{
-		const FVulkanComputePipelineDescriptorInfo& DescriptorInfo = CurrentState->GetComputePipelineDescriptorInfo();
-		uint8 DescriptorSet;
-		uint32 BindingIndex;
-		DescriptorInfo.GetDescriptorSetAndBindingIndex(SRVIndex, DescriptorSet, BindingIndex);
-		SetSRVForUBResource(DescriptorSet, BindingIndex, SRV);
+		SetSRVForUBResource(ShaderStage::Compute, SRVIndex, SRV);
 	}
 
 	inline void SetPackedGlobalShaderParameter(uint32 BufferIndex, uint32 Offset, uint32 NumBytes, const void* NewValue)
@@ -256,11 +240,7 @@ public:
 
 	inline void SetTextureForStage(ShaderStage::EStage Stage, uint32 ParameterIndex, const FVulkanTexture* Texture, VkImageLayout Layout)
 	{
-		const FVulkanGfxPipelineDescriptorInfo& DescriptorInfo = CurrentState->GetGfxPipelineDescriptorInfo();
-		uint8 DescriptorSet;
-		uint32 BindingIndex;
-		DescriptorInfo.GetDescriptorSetAndBindingIndex(Stage, ParameterIndex, DescriptorSet, BindingIndex);
-		CurrentState->SetTexture(DescriptorSet, BindingIndex, Texture, Layout);
+		CurrentState->SetTexture(Stage, ParameterIndex, Texture, Layout);
 	}
 
 	inline void SetTextureForUBResource(uint8 DescriptorSet, uint32 BindingIndex, const FVulkanTexture* Texture, VkImageLayout Layout)
@@ -278,31 +258,19 @@ public:
 
 	inline void SetUAVForStage(ShaderStage::EStage Stage, uint32 ParameterIndex, FVulkanUnorderedAccessView* UAV)
 	{
-		const FVulkanGfxPipelineDescriptorInfo& DescriptorInfo = CurrentState->GetGfxPipelineDescriptorInfo();
-		uint8 DescriptorSet;
-		uint32 BindingIndex;
-		DescriptorInfo.GetDescriptorSetAndBindingIndex(Stage, ParameterIndex, DescriptorSet, BindingIndex);
-		SetUAVForUBResource(DescriptorSet, BindingIndex, UAV);
+		SetUAVForUBResource(Stage, ParameterIndex, UAV);
 	}
 
 	void SetSRVForUBResource(uint8 DescriptorSet, uint32 BindingIndex, FVulkanShaderResourceView* SRV);
 
 	inline void SetSRVForStage(ShaderStage::EStage Stage, uint32 ParameterIndex, FVulkanShaderResourceView* SRV)
 	{
-		const FVulkanGfxPipelineDescriptorInfo& DescriptorInfo = CurrentState->GetGfxPipelineDescriptorInfo();
-		uint8 DescriptorSet;
-		uint32 BindingIndex;
-		DescriptorInfo.GetDescriptorSetAndBindingIndex(Stage, ParameterIndex, DescriptorSet, BindingIndex);
-		SetSRVForUBResource(DescriptorSet, BindingIndex, SRV);
+		SetSRVForUBResource(Stage, ParameterIndex, SRV);
 	}
 
 	inline void SetSamplerStateForStage(ShaderStage::EStage Stage, uint32 ParameterIndex, FVulkanSamplerState* Sampler)
 	{
-		const FVulkanGfxPipelineDescriptorInfo& DescriptorInfo = CurrentState->GetGfxPipelineDescriptorInfo();
-		uint8 DescriptorSet;
-		uint32 BindingIndex;
-		DescriptorInfo.GetDescriptorSetAndBindingIndex(Stage, ParameterIndex, DescriptorSet, BindingIndex);
-		CurrentState->SetSamplerState(DescriptorSet, BindingIndex, Sampler);
+		CurrentState->SetSamplerState(Stage, ParameterIndex, Sampler);
 	}
 
 	inline void SetSamplerStateForUBResource(uint32 DescriptorSet, uint32 BindingIndex, FVulkanSamplerState* Sampler)
diff --git a/Engine/Source/Runtime/VulkanRHI/Private/VulkanPipelineState.cpp b/Engine/Source/Runtime/VulkanRHI/Private/VulkanPipelineState.cpp
--- a/Engine/Source/Runtime/VulkanRHI/Private/VulkanPipelineState.cpp
+++ b/Engine/Source/Runtime/VulkanRHI/Private/VulkanPipelineState.cpp
@@ -178,7 +178,7 @@ bool FVulkanComputePipelineDescriptorState::InternalUpdateDescriptorSets(FVulkan
 #if VULKAN_ENABLE_AGGRESSIVE_STATS
 		SCOPE_CYCLE_COUNTER(STAT_VulkanApplyPackedUniformBuffers);
 #endif
-		UpdatePackedUniformBuffers<bUseDynamicGlobalUBs>(UBOffsetAlignment, PipelineDescriptorInfo->RemappingInfo->StageInfos[0].PackedUBBindingIndices.GetData(), PackedUniformBuffers, DSWriter[0], UniformBufferUploader, CPURingBufferBase, PackedUniformBuffersDirty, CmdBuffer);
+		UpdatePackedUniformBuffers<bUseDynamicGlobalUBs>(UBOffsetAlignment, PackedUniformBuffers, DSWriter[ShaderStage::Compute], UniformBufferUploader, CPURingBufferBase, CmdBuffer);
 		PackedUniformBuffersDirty = 0;
 	}
 
@@ -230,53 +230,38 @@ void FVulkanComputePipelineDescriptorState::UpdateBindlessDescriptors(FVulkanCom
 	const FDescriptorSetRemappingInfo::FStageInfo& StageInfo = RemappingInfo->StageInfos[Stage];
 
 	TArray<VkDescriptorAddressInfoEXT>& DescriptorAddressInfos = StageUBs[Stage];
-	DescriptorAddressInfos.SetNumZeroed(StageInfo.PackedUBBindingIndices.Num() + StageInfo.UniformBuffers.Num());
+	DescriptorAddressInfos.SetNumZeroed((StageInfo.PackedGlobalsSize ? 1 : 0) + StageInfo.NumBoundUniformBuffers);
 
 	// PackedUniformBuffersDirty ?
-	check((PackedUniformBuffersMask == 0) || (PackedUniformBuffersMask == 1));
+	check(PackedUniformBuffersMask <= 1);
 	if (PackedUniformBuffersMask != 0)
 	{
-		check((StageInfo.PackedUBDescriptorSet != UINT16_MAX) && (StageInfo.PackedUBDescriptorSet <= (uint16)Stage));
-		uint64 RemainingPackedUniformsMask = PackedUniformBuffersMask;
-		int32 PackedUBIndex = 0;
-		while (RemainingPackedUniformsMask)
-		{
-			if (RemainingPackedUniformsMask & 1)
-			{
-				const FPackedUniformBuffers::FPackedBuffer& StagedUniformBuffer = PackedUniformBuffers.GetBuffer(PackedUBIndex);
-				const int32 UBSize = StagedUniformBuffer.Num();
-				const int32 BindingIndex = StageInfo.PackedUBBindingIndices[PackedUBIndex];
+		const FPackedUniformBuffers::FPackedBuffer& StagedUniformBuffer = PackedUniformBuffers.GetBuffer();
+		const int32 UBSize = StagedUniformBuffer.Num();
+		const int32 BindingIndex = 0;
 
-				const uint64 RingBufferOffset = UniformBufferUploader->AllocateMemory(UBSize, UBOffsetAlignment, CmdBuffer);
+		const uint64 RingBufferOffset = UniformBufferUploader->AllocateMemory(UBSize, UBOffsetAlignment, CmdBuffer);
 
-				// Make sure it wasn't written to already
-				VkDescriptorAddressInfoEXT& DescriptorAddressInfo = DescriptorAddressInfos[BindingIndex];
-				check(DescriptorAddressInfo.sType == 0);
-				DescriptorAddressInfo.sType = VK_STRUCTURE_TYPE_DESCRIPTOR_ADDRESS_INFO_EXT;
-				DescriptorAddressInfo.address = UniformBufferUploader->GetCPUBufferAddress() + RingBufferOffset;
-				DescriptorAddressInfo.range = UBSize;
+		// Make sure it wasn't written to already
+		VkDescriptorAddressInfoEXT& DescriptorAddressInfo = DescriptorAddressInfos[BindingIndex];
+		check(DescriptorAddressInfo.sType == 0);
+		DescriptorAddressInfo.sType = VK_STRUCTURE_TYPE_DESCRIPTOR_ADDRESS_INFO_EXT;
+		DescriptorAddressInfo.address = UniformBufferUploader->GetCPUBufferAddress() + RingBufferOffset;
+		DescriptorAddressInfo.range = UBSize;
 
-				// get location in the ring buffer to use
-				FMemory::Memcpy(CPURingBufferBase + RingBufferOffset, StagedUniformBuffer.GetData(), UBSize);
-			}
-			RemainingPackedUniformsMask = RemainingPackedUniformsMask >> 1;
-			++PackedUBIndex;
-		}
+		// get location in the ring buffer to use
+		FMemory::Memcpy(CPURingBufferBase + RingBufferOffset, StagedUniformBuffer.GetData(), UBSize);
 
 		PackedUniformBuffersDirty = 0;
 	}
 
-	for (int32 UBIndex = 0; UBIndex < StageInfo.UniformBuffers.Num(); ++UBIndex)
+	for (uint32 UBIndex = 0; UBIndex < StageInfo.NumBoundUniformBuffers; ++UBIndex)
 	{
-		const FDescriptorSetRemappingInfo::FUBRemappingInfo& UBRemappingInfo = StageInfo.UniformBuffers[UBIndex];
-		check((StageInfo.PackedUBDescriptorSet == UINT16_MAX) || UBRemappingInfo.Remapping.NewDescriptorSet == StageInfo.PackedUBDescriptorSet);
-		check(UBRemappingInfo.bHasConstantData);
-
-		VkDescriptorAddressInfoEXT& DescriptorAddressInfo = DescriptorAddressInfos[UBRemappingInfo.Remapping.NewBindingIndex];
+		VkDescriptorAddressInfoEXT& DescriptorAddressInfo = DescriptorAddressInfos[UBIndex];
 		check(DescriptorAddressInfo.sType == 0);
 
-		VkWriteDescriptorSet& WriteDescriptorSet = DSWriter[UBRemappingInfo.Remapping.NewDescriptorSet].WriteDescriptors[UBRemappingInfo.Remapping.NewBindingIndex];
-		check(WriteDescriptorSet.dstBinding == UBRemappingInfo.Remapping.NewBindingIndex);
+		VkWriteDescriptorSet& WriteDescriptorSet = DSWriter[Stage].WriteDescriptors[UBIndex];
+		check(WriteDescriptorSet.dstBinding == UBIndex);
 		check(WriteDescriptorSet.dstArrayElement == 0);
 		check(WriteDescriptorSet.descriptorCount == 1);
 		check(WriteDescriptorSet.descriptorType == VK_DESCRIPTOR_TYPE_UNIFORM_BUFFER);
@@ -364,8 +349,6 @@ bool FVulkanGraphicsPipelineDescriptorState::InternalUpdateDescriptorSets(FVulka
 	uint8* CPURingBufferBase = (uint8*)UniformBufferUploader->GetCPUMappedPointer();
 	const VkDeviceSize UBOffsetAlignment = Device->GetLimits().minUniformBufferOffsetAlignment;
 
-	const FDescriptorSetRemappingInfo* RESTRICT RemappingInfo = PipelineDescriptorInfo->RemappingInfo;
-
 	// Process updates
 	{
 #if VULKAN_ENABLE_AGGRESSIVE_STATS
@@ -375,8 +358,7 @@ bool FVulkanGraphicsPipelineDescriptorState::InternalUpdateDescriptorSets(FVulka
 		{
 			if (PackedUniformBuffersDirty[Stage] != 0)
 			{
-				const uint32 DescriptorSet = RemappingInfo->StageInfos[Stage].PackedUBDescriptorSet;
-				MarkDirty(UpdatePackedUniformBuffers<bUseDynamicGlobalUBs>(UBOffsetAlignment, RemappingInfo->StageInfos[Stage].PackedUBBindingIndices.GetData(), PackedUniformBuffers[Stage], DSWriter[DescriptorSet], UniformBufferUploader, CPURingBufferBase, PackedUniformBuffersDirty[Stage], CmdBuffer));
+				MarkDirty(UpdatePackedUniformBuffers<bUseDynamicGlobalUBs>(UBOffsetAlignment, PackedUniformBuffers[Stage], DSWriter[Stage], UniformBufferUploader, CPURingBufferBase, CmdBuffer));
 				PackedUniformBuffersDirty[Stage] = 0;
 			}
 		}
@@ -456,53 +438,38 @@ void FVulkanGraphicsPipelineDescriptorState::UpdateBindlessDescriptors(FVulkanCo
 			const FDescriptorSetRemappingInfo::FStageInfo& StageInfo = RemappingInfo->StageInfos[Stage];
 
 			TArray<VkDescriptorAddressInfoEXT>& DescriptorAddressInfos = StageUBs[Stage];
-			DescriptorAddressInfos.SetNumZeroed(StageInfo.PackedUBBindingIndices.Num() + StageInfo.UniformBuffers.Num());
+			DescriptorAddressInfos.SetNumZeroed((StageInfo.PackedGlobalsSize ? 1 : 0) + StageInfo.NumBoundUniformBuffers);
 
 			// PackedUniformBuffersDirty ?
-			check((PackedUniformBuffersMask[Stage] == 0) || (PackedUniformBuffersMask[Stage] == 1));
+			check(PackedUniformBuffersMask[Stage] <= 1);
 			if (PackedUniformBuffersMask[Stage] != 0)
 			{
-				check((StageInfo.PackedUBDescriptorSet != UINT16_MAX) && (StageInfo.PackedUBDescriptorSet <= (uint16)Stage));
-				uint64 RemainingPackedUniformsMask = PackedUniformBuffersMask[Stage];
-				int32 PackedUBIndex = 0;
-				while (RemainingPackedUniformsMask)
-				{
-					if (RemainingPackedUniformsMask & 1)
-					{
-						const FPackedUniformBuffers::FPackedBuffer& StagedUniformBuffer = PackedUniformBuffers[Stage].GetBuffer(PackedUBIndex);
-						const int32 UBSize = StagedUniformBuffer.Num();
-						const int32 BindingIndex = StageInfo.PackedUBBindingIndices[PackedUBIndex];
-
-						const uint64 RingBufferOffset = UniformBufferUploader->AllocateMemory(UBSize, UBOffsetAlignment, CmdBuffer);
-
-						// Make sure it wasn't written to already
-						VkDescriptorAddressInfoEXT& DescriptorAddressInfo = DescriptorAddressInfos[BindingIndex];
-						check(DescriptorAddressInfo.sType == 0);
-						DescriptorAddressInfo.sType = VK_STRUCTURE_TYPE_DESCRIPTOR_ADDRESS_INFO_EXT;
-						DescriptorAddressInfo.address = UniformBufferUploader->GetCPUBufferAddress() + RingBufferOffset;
-						DescriptorAddressInfo.range = UBSize;
-
-						// get location in the ring buffer to use
-						FMemory::Memcpy(CPURingBufferBase + RingBufferOffset, StagedUniformBuffer.GetData(), UBSize);
-					}
-					RemainingPackedUniformsMask = RemainingPackedUniformsMask >> 1;
-					++PackedUBIndex;
-				}
+				const FPackedUniformBuffers::FPackedBuffer& StagedUniformBuffer = PackedUniformBuffers[Stage].GetBuffer();
+				const int32 UBSize = StagedUniformBuffer.Num();
+				const int32 BindingIndex = 0;
+
+				const uint64 RingBufferOffset = UniformBufferUploader->AllocateMemory(UBSize, UBOffsetAlignment, CmdBuffer);
+
+				// Make sure it wasn't written to already
+				VkDescriptorAddressInfoEXT& DescriptorAddressInfo = DescriptorAddressInfos[BindingIndex];
+				check(DescriptorAddressInfo.sType == 0);
+				DescriptorAddressInfo.sType = VK_STRUCTURE_TYPE_DESCRIPTOR_ADDRESS_INFO_EXT;
+				DescriptorAddressInfo.address = UniformBufferUploader->GetCPUBufferAddress() + RingBufferOffset;
+				DescriptorAddressInfo.range = UBSize;
+
+				// get location in the ring buffer to use
+				FMemory::Memcpy(CPURingBufferBase + RingBufferOffset, StagedUniformBuffer.GetData(), UBSize);
 
 				PackedUniformBuffersDirty[Stage] = 0;
 			}
 
-			for (int32 UBIndex = 0; UBIndex < StageInfo.UniformBuffers.Num(); ++UBIndex)
+			for (uint32 UBIndex = 0; UBIndex < StageInfo.NumBoundUniformBuffers; ++UBIndex)
 			{
-				const FDescriptorSetRemappingInfo::FUBRemappingInfo& UBRemappingInfo = StageInfo.UniformBuffers[UBIndex];
-				check((StageInfo.PackedUBDescriptorSet == UINT16_MAX) || UBRemappingInfo.Remapping.NewDescriptorSet == StageInfo.PackedUBDescriptorSet);
-				check(UBRemappingInfo.bHasConstantData);
-
-				VkDescriptorAddressInfoEXT& DescriptorAddressInfo = DescriptorAddressInfos[UBRemappingInfo.Remapping.NewBindingIndex];
+				VkDescriptorAddressInfoEXT& DescriptorAddressInfo = DescriptorAddressInfos[UBIndex];
 				check(DescriptorAddressInfo.sType == 0);
 
-				VkWriteDescriptorSet& WriteDescriptorSet = DSWriter[UBRemappingInfo.Remapping.NewDescriptorSet].WriteDescriptors[UBRemappingInfo.Remapping.NewBindingIndex];
-				check(WriteDescriptorSet.dstBinding == UBRemappingInfo.Remapping.NewBindingIndex);
+				VkWriteDescriptorSet& WriteDescriptorSet = DSWriter[Stage].WriteDescriptors[UBIndex];
+				check(WriteDescriptorSet.dstBinding == UBIndex);
 				check(WriteDescriptorSet.dstArrayElement == 0);
 				check(WriteDescriptorSet.descriptorCount == 1);
 				check(WriteDescriptorSet.descriptorType == VK_DESCRIPTOR_TYPE_UNIFORM_BUFFER);
@@ -531,7 +498,7 @@ void FVulkanCommandListContext::ApplyStaticUniformBuffers(TRHIShader* Shader)
 	if (Shader)
 	{
 		const auto& StaticSlots = Shader->GetStaticSlots();
-		const auto& UBInfos = Shader->GetCodeHeader().UniformBuffers;
+		const auto& UBInfos = Shader->GetCodeHeader().UniformBufferInfos;
 
 		for (int32 BufferIndex = 0; BufferIndex < StaticSlots.Num(); ++BufferIndex)
 		{
diff --git a/Engine/Source/Runtime/VulkanRHI/Private/VulkanPipelineState.h b/Engine/Source/Runtime/VulkanRHI/Private/VulkanPipelineState.h
--- a/Engine/Source/Runtime/VulkanRHI/Private/VulkanPipelineState.h
+++ b/Engine/Source/Runtime/VulkanRHI/Private/VulkanPipelineState.h
@@ -186,7 +186,8 @@ public:
 
 	inline void SetPackedGlobalShaderParameter(uint32 BufferIndex, uint32 ByteOffset, uint32 NumBytes, const void* NewValue)
 	{
-		PackedUniformBuffers.SetPackedGlobalParameter(BufferIndex, ByteOffset, NumBytes, NewValue, PackedUniformBuffersDirty);
+		check(BufferIndex == 0);
+		PackedUniformBuffers.SetPackedGlobalParameter(ByteOffset, NumBytes, NewValue, PackedUniformBuffersDirty);
 	}
 
 	bool UpdateDescriptorSets(FVulkanCommandListContext* CmdListContext, FVulkanCmdBuffer* CmdBuffer)
@@ -221,8 +222,8 @@ protected:
 	const FVulkanComputePipelineDescriptorInfo* PipelineDescriptorInfo;
 
 	FPackedUniformBuffers PackedUniformBuffers;
-	uint64 PackedUniformBuffersMask;
-	uint64 PackedUniformBuffersDirty;
+	uint32 PackedUniformBuffersMask;
+	uint32 PackedUniformBuffersDirty;
 
 	FVulkanComputePipeline* ComputePipeline;
 
@@ -244,7 +245,8 @@ public:
 
 	inline void SetPackedGlobalShaderParameter(uint8 Stage, uint32 BufferIndex, uint32 ByteOffset, uint32 NumBytes, const void* NewValue)
 	{
-		PackedUniformBuffers[Stage].SetPackedGlobalParameter(BufferIndex, ByteOffset, NumBytes, NewValue, PackedUniformBuffersDirty[Stage]);
+		check(BufferIndex == 0);
+		PackedUniformBuffers[Stage].SetPackedGlobalParameter(ByteOffset, NumBytes, NewValue, PackedUniformBuffersDirty[Stage]);
 	}
 
 	bool UpdateDescriptorSets(FVulkanCommandListContext* CmdListContext, FVulkanCmdBuffer* CmdBuffer)
@@ -285,8 +287,8 @@ protected:
 	const FVulkanGfxPipelineDescriptorInfo* PipelineDescriptorInfo;
 
 	TStaticArray<FPackedUniformBuffers, ShaderStage::NumStages> PackedUniformBuffers;
-	TStaticArray<uint64, ShaderStage::NumStages> PackedUniformBuffersMask;
-	TStaticArray<uint64, ShaderStage::NumStages> PackedUniformBuffersDirty;
+	TStaticArray<uint32, ShaderStage::NumStages> PackedUniformBuffersMask;
+	TStaticArray<uint32, ShaderStage::NumStages> PackedUniformBuffersDirty;
 
 	FVulkanRHIGraphicsPipelineState* GfxPipeline;
 
@@ -298,44 +300,28 @@ protected:
 };
 
 template <bool bIsDynamic>
-static inline bool UpdatePackedUniformBuffers(VkDeviceSize UBOffsetAlignment, const uint16* RESTRICT PackedUBBindingIndices, const FPackedUniformBuffers& PackedUniformBuffers,
-	FVulkanDescriptorSetWriter& DescriptorWriteSet, FVulkanUniformBufferUploader* UniformBufferUploader, uint8* RESTRICT CPURingBufferBase, uint64 RemainingPackedUniformsMask,
+static inline bool UpdatePackedUniformBuffers(VkDeviceSize UBOffsetAlignment, const FPackedUniformBuffers& PackedUniformBuffers,
+	FVulkanDescriptorSetWriter& DescriptorWriteSet, FVulkanUniformBufferUploader* UniformBufferUploader, uint8* RESTRICT CPURingBufferBase,
 	FVulkanCmdBuffer* InCmdBuffer)
 {
-	bool bAnyUBDirty = false;
-	int32 PackedUBIndex = 0;
-	while (RemainingPackedUniformsMask)
-	{
-		if (RemainingPackedUniformsMask & 1)
-		{
-			const FPackedUniformBuffers::FPackedBuffer& StagedUniformBuffer = PackedUniformBuffers.GetBuffer(PackedUBIndex);
-			int32 BindingIndex = PackedUBBindingIndices[PackedUBIndex];
+	const FPackedUniformBuffers::FPackedBuffer& StagedUniformBuffer = PackedUniformBuffers.GetBuffer();
+	const int32 BindingIndex = 0;  // Packed unifrom buffers are only used for globals at binding 0
 
-			const int32 UBSize = StagedUniformBuffer.Num();
+	const int32 UBSize = StagedUniformBuffer.Num();
 
-			// get offset into the RingBufferBase pointer
-			uint64 RingBufferOffset = UniformBufferUploader->AllocateMemory(UBSize, UBOffsetAlignment, InCmdBuffer);
+	// get offset into the RingBufferBase pointer
+	uint64 RingBufferOffset = UniformBufferUploader->AllocateMemory(UBSize, UBOffsetAlignment, InCmdBuffer);
 
-			// get location in the ring buffer to use
-			FMemory::Memcpy(CPURingBufferBase + RingBufferOffset, StagedUniformBuffer.GetData(), UBSize);
-
-			const VulkanRHI::FVulkanAllocation& Allocation = UniformBufferUploader->GetCPUBufferAllocation();
-			if (bIsDynamic)
-			{
-				const bool bDirty = DescriptorWriteSet.WriteDynamicUniformBuffer(BindingIndex, Allocation.GetBufferHandle(), Allocation.HandleId, UniformBufferUploader->GetCPUBufferOffset(), UBSize, RingBufferOffset);
-				bAnyUBDirty = bAnyUBDirty || bDirty;
-
-			}
-			else
-			{
-				const bool bDirty = DescriptorWriteSet.WriteUniformBuffer(BindingIndex, Allocation.GetBufferHandle(), Allocation.HandleId, RingBufferOffset + UniformBufferUploader->GetCPUBufferOffset(), UBSize);
-				bAnyUBDirty = bAnyUBDirty || bDirty;
+	// get location in the ring buffer to use
+	FMemory::Memcpy(CPURingBufferBase + RingBufferOffset, StagedUniformBuffer.GetData(), UBSize);
 
-			}
-		}
-		RemainingPackedUniformsMask = RemainingPackedUniformsMask >> 1;
-		++PackedUBIndex;
+	const VulkanRHI::FVulkanAllocation& Allocation = UniformBufferUploader->GetCPUBufferAllocation();
+	if (bIsDynamic)
+	{
+		return DescriptorWriteSet.WriteDynamicUniformBuffer(BindingIndex, Allocation.GetBufferHandle(), Allocation.HandleId, UniformBufferUploader->GetCPUBufferOffset(), UBSize, RingBufferOffset);
+	}
+	else
+	{
+		return DescriptorWriteSet.WriteUniformBuffer(BindingIndex, Allocation.GetBufferHandle(), Allocation.HandleId, RingBufferOffset + UniformBufferUploader->GetCPUBufferOffset(), UBSize);
 	}
-
-	return bAnyUBDirty;
 }
diff --git a/Engine/Source/Runtime/VulkanRHI/Private/VulkanRHI.cpp b/Engine/Source/Runtime/VulkanRHI/Private/VulkanRHI.cpp
--- a/Engine/Source/Runtime/VulkanRHI/Private/VulkanRHI.cpp
+++ b/Engine/Source/Runtime/VulkanRHI/Private/VulkanRHI.cpp
@@ -1704,22 +1704,14 @@ void FVulkanDescriptorSetsLayoutInfo::GenerateHash(const TArrayView<FRHISamplerS
 	{
 		const FDescriptorSetRemappingInfo::FStageInfo& StageInfo = RemappingInfo.StageInfos[RemapingIndex];
 
-		Hash = FCrc::MemCrc32(&StageInfo.PackedUBDescriptorSet, sizeof(uint16), Hash);
-		Hash = FCrc::MemCrc32(&StageInfo.NumImageInfos, sizeof(uint16), Hash);
-		Hash = FCrc::MemCrc32(&StageInfo.NumBufferInfos, sizeof(uint16), Hash);
-		Hash = FCrc::MemCrc32(&StageInfo.NumAccelerationStructures, sizeof(uint16), Hash);
-
-		const TArray<FDescriptorSetRemappingInfo::FRemappingInfo>& Globals = StageInfo.Globals;
-		Hash = FCrc::MemCrc32(Globals.GetData(), sizeof(FDescriptorSetRemappingInfo::FRemappingInfo) * Globals.Num(), Hash);
-
-		const TArray<FDescriptorSetRemappingInfo::FUBRemappingInfo>& UniformBuffers = StageInfo.UniformBuffers;
-		Hash = FCrc::MemCrc32(UniformBuffers.GetData(), sizeof(FDescriptorSetRemappingInfo::FUBRemappingInfo) * UniformBuffers.Num(), Hash);
+		Hash = FCrc::TypeCrc32(StageInfo.PackedGlobalsSize, Hash);
+		Hash = FCrc::TypeCrc32(StageInfo.NumBoundUniformBuffers, Hash);
+		Hash = FCrc::TypeCrc32(StageInfo.NumImageInfos, Hash);
+		Hash = FCrc::TypeCrc32(StageInfo.NumBufferInfos, Hash);
+		Hash = FCrc::TypeCrc32(StageInfo.NumAccelerationStructures, Hash);
 
 		const TArray<VkDescriptorType>& Types = StageInfo.Types;
 		Hash = FCrc::MemCrc32(Types.GetData(), sizeof(VkDescriptorType) * Types.Num(), Hash);
-
-		const TArray<uint16>& PackedUBBindingIndices = StageInfo.PackedUBBindingIndices;
-		Hash = FCrc::MemCrc32(PackedUBBindingIndices.GetData(), sizeof(uint16) * PackedUBBindingIndices.Num(), Hash);
 	}
 
 	// It would be better to store this when the object is created, but it's not available at that time, so we'll do it here.
diff --git a/Engine/Source/Runtime/VulkanRHI/Private/VulkanRHIPrivate.h b/Engine/Source/Runtime/VulkanRHI/Private/VulkanRHIPrivate.h
--- a/Engine/Source/Runtime/VulkanRHI/Private/VulkanRHIPrivate.h
+++ b/Engine/Source/Runtime/VulkanRHI/Private/VulkanRHIPrivate.h
@@ -40,7 +40,6 @@ class FVulkanDescriptorSetsLayout;
 class FVulkanGfxPipeline;
 class FVulkanRenderPass;
 class FVulkanCommandBufferManager;
-struct FInputAttachmentData;
 class FValidationContext;
 
 
diff --git a/Engine/Source/Runtime/VulkanRHI/Private/VulkanRayTracing.cpp b/Engine/Source/Runtime/VulkanRHI/Private/VulkanRayTracing.cpp
--- a/Engine/Source/Runtime/VulkanRHI/Private/VulkanRayTracing.cpp
+++ b/Engine/Source/Runtime/VulkanRHI/Private/VulkanRayTracing.cpp
@@ -2203,12 +2203,12 @@ void FVulkanCommandListContext::RHIRayTraceDispatchIndirect(
 static void SetSystemParametersUB(FVulkanHitGroupSystemParameters& OutSystemParameters, FVulkanRayTracingShaderTable* ShaderTable, uint32 InNumUniformBuffers, FRHIUniformBuffer* const* InUniformBuffers, const FVulkanRayTracingShader* InShader)
 {
 	// Plug the shaders in the right slots using LayoutHash comparisons
-	check(InShader->GetCodeHeader().UniformBuffers.Num() <= (int32)InNumUniformBuffers);
-	for (int32 UBIndex = 0; UBIndex < InShader->GetCodeHeader().UniformBuffers.Num(); ++UBIndex)
+	check(InShader->GetCodeHeader().UniformBufferInfos.Num() <= (int32)InNumUniformBuffers);
+	for (int32 UBIndex = 0; UBIndex < InShader->GetCodeHeader().UniformBufferInfos.Num(); ++UBIndex)
 	{
 		FVulkanUniformBuffer* UniformBuffer = ResourceCast(InUniformBuffers[UBIndex]);
 
-		const FVulkanShaderHeader::FUniformBufferInfo& UniformBufferInfo = InShader->GetCodeHeader().UniformBuffers[UBIndex];
+		const FVulkanShaderHeader::FUniformBufferInfo& UniformBufferInfo = InShader->GetCodeHeader().UniformBufferInfos[UBIndex];
 
 		// :todo-jn: Hack to force in a DummyCullingBuffer in cases where it should have been culled from source (see SPIRV-Tools Issue 4902).
 		if (!UniformBuffer)
@@ -2217,13 +2217,11 @@ static void SetSystemParametersUB(FVulkanHitGroupSystemParameters& OutSystemPara
 		}
 
 		check(UniformBuffer);
-
 		check((UniformBufferInfo.LayoutHash == 0) || (UniformBufferInfo.LayoutHash == UniformBuffer->GetLayout().GetHash()));
-		check(UniformBufferInfo.ConstantDataOriginalBindingIndex != UINT16_MAX);
 
 		const FRHIDescriptorHandle BindlessHandle = UniformBuffer->GetBindlessHandle();
 		check(BindlessHandle.IsValid());
-		OutSystemParameters.BindlessUniformBuffers[UniformBufferInfo.ConstantDataOriginalBindingIndex] = BindlessHandle.GetIndex();
+		OutSystemParameters.BindlessUniformBuffers[UniformBufferInfo.BindlessCBIndex] = BindlessHandle.GetIndex();
 
 		ShaderTable->AddUBRef(UniformBuffer);
 	}
diff --git a/Engine/Source/Runtime/VulkanRHI/Private/VulkanResources.h b/Engine/Source/Runtime/VulkanRHI/Private/VulkanResources.h
--- a/Engine/Source/Runtime/VulkanRHI/Private/VulkanResources.h
+++ b/Engine/Source/Runtime/VulkanRHI/Private/VulkanResources.h
@@ -1283,23 +1283,25 @@ public:
 	// One buffer is a chunk of bytes
 	typedef TArray<uint8> FPackedBuffer;
 
-	void Init(const FVulkanShaderHeader& InCodeHeader, uint64& OutPackedUniformBufferStagingMask)
+	void Init(const FVulkanShaderHeader& InCodeHeader, uint32& OutPackedUniformBufferStagingMask)
 	{
-		PackedUniformBuffers.AddDefaulted(InCodeHeader.PackedUBs.Num());
-		for (int32 Index = 0; Index < InCodeHeader.PackedUBs.Num(); ++Index)
+		if (InCodeHeader.PackedGlobalsSize > 0)
 		{
-			PackedUniformBuffers[Index].AddUninitialized(InCodeHeader.PackedUBs[Index].SizeInBytes);
+			check(PackedUniformBuffers.Num() == 0);
+			PackedUniformBuffers.AddUninitialized(InCodeHeader.PackedGlobalsSize);
+			OutPackedUniformBufferStagingMask = 1;
+		}
+		else
+		{
+			OutPackedUniformBufferStagingMask = 0;
 		}
-
-		OutPackedUniformBufferStagingMask = ((uint64)1 << (uint64)InCodeHeader.PackedUBs.Num()) - 1;
 	}
 
-	inline void SetPackedGlobalParameter(uint32 BufferIndex, uint32 ByteOffset, uint32 NumBytes, const void* RESTRICT NewValue, uint64& InOutPackedUniformBufferStagingDirty)
+	inline void SetPackedGlobalParameter(uint32 ByteOffset, uint32 NumBytes, const void* RESTRICT NewValue, uint32& InOutPackedUniformBufferStagingDirty)
 	{
-		FPackedBuffer& StagingBuffer = PackedUniformBuffers[BufferIndex];
-		check(ByteOffset + NumBytes <= (uint32)StagingBuffer.Num());
+		check(ByteOffset + NumBytes <= (uint32)PackedUniformBuffers.Num());
 		check((NumBytes & 3) == 0 && (ByteOffset & 3) == 0);
-		uint32* RESTRICT RawDst = (uint32*)(StagingBuffer.GetData() + ByteOffset);
+		uint32* RESTRICT RawDst = (uint32*)(PackedUniformBuffers.GetData() + ByteOffset);
 		uint32* RESTRICT RawSrc = (uint32*)NewValue;
 		uint32* RESTRICT RawSrcEnd = RawSrc + (NumBytes >> 2);
 
@@ -1309,21 +1311,19 @@ public:
 			bChanged |= CopyAndReturnNotEqual(*RawDst++, *RawSrc++);
 		}
 
-		InOutPackedUniformBufferStagingDirty = InOutPackedUniformBufferStagingDirty | ((uint64)(bChanged ? 1 : 0) << (uint64)BufferIndex);
-	}
-
-	// Copies a 'real' constant buffer into the packed globals uniform buffer (only the used ranges)
-	inline void SetEmulatedUniformBufferIntoPacked(uint32 BindPoint, const TArray<uint8>& ConstantData, uint64& NEWPackedUniformBufferStagingDirty)
-	{
+		if (bChanged)
+		{
+			InOutPackedUniformBufferStagingDirty = 1;
+		}
 	}
 
-	inline const FPackedBuffer& GetBuffer(int32 Index) const
+	inline const FPackedBuffer& GetBuffer() const
 	{
-		return PackedUniformBuffers[Index];
+		return PackedUniformBuffers;
 	}
 
 protected:
-	TArray<FPackedBuffer>									PackedUniformBuffers;
+	FPackedBuffer PackedUniformBuffers;
 };
 
 class FVulkanStagingBuffer : public FRHIStagingBuffer
diff --git a/Engine/Source/Runtime/VulkanRHI/Private/VulkanShaders.cpp b/Engine/Source/Runtime/VulkanRHI/Private/VulkanShaders.cpp
--- a/Engine/Source/Runtime/VulkanRHI/Private/VulkanShaders.cpp
+++ b/Engine/Source/Runtime/VulkanRHI/Private/VulkanShaders.cpp
@@ -276,13 +276,9 @@ FVulkanShader::FVulkanShader(FVulkanDevice* InDevice, EShaderFrequency InFrequen
 
 	checkf(SpirvContainer.GetSizeBytes() != 0, TEXT("Empty SPIR-V! %s"), *CodeHeader.DebugName);
 
-	check(IsRayTracingShaderFrequency(Frequency) || (CodeHeader.UniformBufferSpirvInfos.Num() == CodeHeader.UniformBuffers.Num()));
+	StaticSlots.Reserve(CodeHeader.UniformBufferInfos.Num());
 
-	check(CodeHeader.GlobalSpirvInfos.Num() == CodeHeader.Globals.Num());
-
-	StaticSlots.Reserve(CodeHeader.UniformBuffers.Num());
-
-	for (const FVulkanShaderHeader::FUniformBufferInfo& UBInfo : CodeHeader.UniformBuffers)
+	for (const FVulkanShaderHeader::FUniformBufferInfo& UBInfo : CodeHeader.UniformBufferInfos)
 	{
 		if (const FShaderParametersMetadata* Metadata = FindUniformBufferStructByLayoutHash(UBInfo.LayoutHash))
 		{
@@ -411,7 +407,7 @@ FVulkanShader::FSpirvCode FVulkanShader::PatchSpirvInputAttachments(FVulkanShade
 
 bool FVulkanShader::NeedsSpirvInputAttachmentPatching(const FGfxPipelineDesc& Desc) const
 {
-	return (Desc.RasterizationSamples > 1 && CodeHeader.InputAttachments.Num() > 0);
+	return (Desc.RasterizationSamples > 1 && CodeHeader.InputAttachmentsMask != 0);
 }
 
 TRefCountPtr<FVulkanShaderModule> FVulkanShader::CreateHandle(const FGfxPipelineDesc& Desc, const FVulkanLayout* Layout, uint32 LayoutHash)
@@ -428,7 +424,6 @@ FVulkanShader::FSpirvCode FVulkanShader::GetPatchedSpirvCode(const FGfxPipelineD
 {
 	FSpirvCode Spirv = GetSpirvCode();
 
-	Layout->PatchSpirvBindings(Spirv, Frequency, CodeHeader);
 	if (NeedsSpirvInputAttachmentPatching(Desc))
 	{
 		Spirv = PatchSpirvInputAttachments(Spirv);
@@ -535,7 +530,6 @@ TRefCountPtr<FVulkanShaderModule> FVulkanShader::CreateHandle(const FVulkanLayou
 	FScopeLock Lock(&VulkanShaderModulesMapCS);
 	FSpirvCode Spirv = GetSpirvCode();
 
-	Layout->PatchSpirvBindings(Spirv, Frequency, CodeHeader);
 	TRefCountPtr<FVulkanShaderModule> Module = CreateShaderModule(Device, Spirv);
 	ShaderModules.Add(LayoutHash, Module);
 	if (!CodeHeader.DebugName.IsEmpty())
@@ -557,62 +551,6 @@ void FVulkanShader::PurgeShaderModules()
 	ShaderModules.Empty(0);
 }
 
-void FVulkanLayout::PatchSpirvBindings(FVulkanShader::FSpirvCode& SprivCode, EShaderFrequency Frequency, const FVulkanShaderHeader& CodeHeader) const
-{
-	// Bindless shader compilation already places descriptors and bindings in their fixed values based on stage and descriptor type
-	if (Device->SupportsBindless())
-	{
-		return;
-	}
-
-	TArrayView<uint32> Spirv = SprivCode.GetCodeView();	//#todo-rco: Do we need an actual copy of the SPIR-V?
-	ShaderStage::EStage Stage = ShaderStage::GetStageForFrequency(Frequency);
-	const FDescriptorSetRemappingInfo::FStageInfo& StageInfo = DescriptorSetLayout.RemappingInfo.StageInfos[Stage];
-
-	checkSlow(StageInfo.UniformBuffers.Num() == CodeHeader.UniformBufferSpirvInfos.Num());
-	for (int32 Index = 0; Index < CodeHeader.UniformBufferSpirvInfos.Num(); ++Index)
-	{
-		if (StageInfo.UniformBuffers[Index].bHasConstantData)
-		{
-			const uint32 OffsetDescriptorSet = CodeHeader.UniformBufferSpirvInfos[Index].DescriptorSetOffset;
-			const uint32 OffsetBindingIndex = CodeHeader.UniformBufferSpirvInfos[Index].BindingIndexOffset;
-			check(OffsetDescriptorSet != UINT32_MAX && OffsetBindingIndex != UINT32_MAX);
-			const uint16 NewDescriptorSet = StageInfo.UniformBuffers[Index].Remapping.NewDescriptorSet;
-			checkf(Spirv[OffsetDescriptorSet] == NewDescriptorSet,
-				TEXT("Descriptor set in SPIRV (%u) differs from expected stage value (%u)"),
-				Spirv[OffsetDescriptorSet], NewDescriptorSet);
-			const uint16 NewBindingIndex = StageInfo.UniformBuffers[Index].Remapping.NewBindingIndex;
-			Spirv[OffsetBindingIndex] = NewBindingIndex;
-		}
-	}
-
-	checkSlow(StageInfo.Globals.Num() == CodeHeader.GlobalSpirvInfos.Num());
-	for (int32 Index = 0; Index < CodeHeader.GlobalSpirvInfos.Num(); ++Index)
-	{
-		const uint32 OffsetDescriptorSet = CodeHeader.GlobalSpirvInfos[Index].DescriptorSetOffset;
-		const uint32 OffsetBindingIndex = CodeHeader.GlobalSpirvInfos[Index].BindingIndexOffset;
-		check(OffsetDescriptorSet != UINT32_MAX && OffsetBindingIndex != UINT32_MAX);
-		const uint16 NewDescriptorSet = StageInfo.Globals[Index].NewDescriptorSet;
-		checkf(Spirv[OffsetDescriptorSet] == NewDescriptorSet,
-			TEXT("Descriptor set in SPIRV (%u) differs from expected stage value (%u)"),
-			Spirv[OffsetDescriptorSet], NewDescriptorSet);
-		const uint16 NewBindingIndex = StageInfo.Globals[Index].NewBindingIndex;
-		Spirv[OffsetBindingIndex] = NewBindingIndex;
-	}
-
-	checkSlow(StageInfo.PackedUBBindingIndices.Num() == CodeHeader.PackedUBs.Num());
-	for (int32 Index = 0; Index < CodeHeader.PackedUBs.Num(); ++Index)
-	{
-		const uint32 OffsetDescriptorSet = CodeHeader.PackedUBs[Index].SPIRVDescriptorSetOffset;
-		const uint32 OffsetBindingIndex = CodeHeader.PackedUBs[Index].SPIRVBindingIndexOffset;
-		check(OffsetDescriptorSet != UINT32_MAX && OffsetBindingIndex != UINT32_MAX);
-		checkf(Spirv[OffsetDescriptorSet] == StageInfo.PackedUBDescriptorSet, 
-			TEXT("Descriptor set in SPIRV (%u) differs from expected stage value (%u)"), 
-			Spirv[OffsetDescriptorSet], StageInfo.PackedUBDescriptorSet);
-		Spirv[OffsetBindingIndex] = StageInfo.PackedUBBindingIndices[Index];
-	}
-}
-
 FVertexShaderRHIRef FVulkanDynamicRHI::RHICreateVertexShader(TArrayView<const uint8> Code, const FSHAHash& Hash)
 {
 	return Device->GetShaderFactory().CreateShader<FVulkanVertexShader>(Code, Device);
@@ -692,8 +630,8 @@ void FVulkanLayout::Compile(FVulkanDescriptorSetLayoutMap& DSetLayoutMap)
 
 bool FVulkanGfxLayout::UsesInputAttachment(FVulkanShaderHeader::EAttachmentType AttachmentType) const
 {
-	const TArray<FInputAttachmentData>& InputAttachmentData = GfxPipelineDescriptorInfo.GetInputAttachmentData();
-	for (const FInputAttachmentData& Input : InputAttachmentData)
+	const TArray<FVulkanShaderHeader::FInputAttachmentInfo>& InputAttachmentData = GfxPipelineDescriptorInfo.GetInputAttachmentData();
+	for (const FVulkanShaderHeader::FInputAttachmentInfo& Input : InputAttachmentData)
 	{
 		if (Input.Type == AttachmentType)
 		{
@@ -801,82 +739,66 @@ void FVulkanDescriptorSetsLayoutInfo::FinalizeBindings(const FVulkanDevice& Devi
 	{
 		if (const FVulkanShaderHeader* ShaderHeader = UBGatherInfo.CodeHeaders[Stage])
 		{
-			// Shaders are compiles with the value of Stage as the descriptor set
-			const int32 DescriptorSet = Stage;
+			FDescriptorSetRemappingInfo::FStageInfo& StageInfo = RemappingInfo.StageInfos[Stage];
+
+			Binding.stageFlags = UEFrequencyToVKStageBit(bIsCompute ? SF_Compute : ShaderStage::GetFrequencyForGfxStage((ShaderStage::EStage)Stage));
 
-			const VkShaderStageFlags StageFlags = UEFrequencyToVKStageBit(bIsCompute ? SF_Compute : ShaderStage::GetFrequencyForGfxStage((ShaderStage::EStage)Stage));
-			Binding.stageFlags = StageFlags;
+			StageInfo.PackedGlobalsSize = ShaderHeader->PackedGlobalsSize;
+			StageInfo.NumBoundUniformBuffers = ShaderHeader->NumBoundUniformBuffers;
 
-			RemappingInfo.StageInfos[Stage].PackedUBBindingIndices.Reserve(ShaderHeader->PackedUBs.Num());
-			for (int32 Index = 0; Index < ShaderHeader->PackedUBs.Num(); ++Index)
+			for (int32 BindingIndex = 0; BindingIndex < ShaderHeader->Bindings.Num(); ++BindingIndex)
 			{
-				const VkDescriptorType Type = bConvertPackedUBsToDynamic ? VK_DESCRIPTOR_TYPE_UNIFORM_BUFFER_DYNAMIC : VK_DESCRIPTOR_TYPE_UNIFORM_BUFFER;
-				const uint32 NewBindingIndex = RemappingInfo.AddPackedUB(Stage, Index, DescriptorSet, Type);
+				const VkDescriptorType DescriptorType = (VkDescriptorType)ShaderHeader->Bindings[BindingIndex].DescriptorType;
 
-				Binding.binding = NewBindingIndex;
-				Binding.descriptorType = Type;
-				AddDescriptor(DescriptorSet, Binding);
-			}
+				const bool bIsUniformBuffer = (DescriptorType == VK_DESCRIPTOR_TYPE_UNIFORM_BUFFER);
+				const bool bIsGlobalPackedConstants = bIsUniformBuffer && ShaderHeader->PackedGlobalsSize && (BindingIndex == 0);
 
-			RemappingInfo.StageInfos[Stage].UniformBuffers.Reserve(ShaderHeader->UniformBuffers.Num());
-			for (int32 Index = 0; Index < ShaderHeader->UniformBuffers.Num(); ++Index)
-			{
-				VkDescriptorType Type = VK_DESCRIPTOR_TYPE_UNIFORM_BUFFER;
-				if (bConvertAllUBsToDynamic && LayoutTypes[VK_DESCRIPTOR_TYPE_UNIFORM_BUFFER_DYNAMIC] < MaxDescriptorSetUniformBuffersDynamic)
+				if (bIsGlobalPackedConstants)
 				{
-					Type = VK_DESCRIPTOR_TYPE_UNIFORM_BUFFER_DYNAMIC;
+					const VkDescriptorType UBType = bConvertPackedUBsToDynamic ? VK_DESCRIPTOR_TYPE_UNIFORM_BUFFER_DYNAMIC : VK_DESCRIPTOR_TYPE_UNIFORM_BUFFER;
+
+					const uint32 NewBindingIndex = StageInfo.Types.Add(UBType);
+					checkf(NewBindingIndex == 0, TEXT("Packed globals should always be the first binding!"));
+
+					Binding.binding = NewBindingIndex;
+					Binding.descriptorType = UBType;
+					AddDescriptor(Stage, Binding);
 				}
-				
-				// Here we might mess up with the stageFlags, so reset them every loop
-				Binding.stageFlags = StageFlags;
-				Binding.descriptorType = Type;
-				const FVulkanShaderHeader::FUniformBufferInfo& UBInfo = ShaderHeader->UniformBuffers[Index];
-				const uint32 LayoutHash = UBInfo.LayoutHash;
-				const bool bUBHasConstantData = UBInfo.ConstantDataOriginalBindingIndex != UINT16_MAX;
-				if (bUBHasConstantData)
+				else if (bIsUniformBuffer)
 				{
-					bool bProcessRegularUB = true;
-					if (bProcessRegularUB)
+					VkDescriptorType UBType = VK_DESCRIPTOR_TYPE_UNIFORM_BUFFER;
+					if (bConvertAllUBsToDynamic && LayoutTypes[VK_DESCRIPTOR_TYPE_UNIFORM_BUFFER_DYNAMIC] < MaxDescriptorSetUniformBuffersDynamic)
 					{
-						uint32 NewBindingIndex;
-						RemappingInfo.AddUBWithData(Stage, Index, DescriptorSet, Type, NewBindingIndex);
-						Binding.binding = NewBindingIndex;
+						UBType = VK_DESCRIPTOR_TYPE_UNIFORM_BUFFER_DYNAMIC;
+					}
 
-						AddDescriptor(DescriptorSet, Binding);
+					// Here we might mess up with the stageFlags, so reset them every loop
+					Binding.descriptorType = UBType;
+					const FVulkanShaderHeader::FUniformBufferInfo& UBInfo = ShaderHeader->UniformBufferInfos[BindingIndex];
+					const bool bUBHasConstantData = (BindingIndex < (int32)ShaderHeader->NumBoundUniformBuffers);
+					if (bUBHasConstantData)
+					{
+						const uint32 NewBindingIndex = StageInfo.Types.Add(UBType);
+						check(NewBindingIndex == BindingIndex);
+						Binding.binding = NewBindingIndex;
+						AddDescriptor(Stage, Binding);
 					}
 				}
 				else
 				{
-					RemappingInfo.AddUBResourceOnly(Stage, Index);
+					const uint32 NewTypeIndex = StageInfo.Types.Add(DescriptorType);
+					check(NewTypeIndex == BindingIndex);
+					Binding.binding = BindingIndex;
+					Binding.descriptorType = DescriptorType;
+					AddDescriptor(Stage, Binding);
 				}
 			}
 
-			RemappingInfo.StageInfos[Stage].Globals.Reserve(ShaderHeader->Globals.Num());
-			Binding.stageFlags = StageFlags;
-			for (int32 Index = 0; Index < ShaderHeader->Globals.Num(); ++Index)
-			{
-				const FVulkanShaderHeader::FGlobalInfo& GlobalInfo = ShaderHeader->Globals[Index];
-				const VkDescriptorType Type = (VkDescriptorType)GlobalInfo.TypeIndex;
-				uint32 NewBindingIndex = RemappingInfo.AddGlobal(Stage, Index, DescriptorSet, Type);
-				Binding.binding = NewBindingIndex;
-				Binding.descriptorType = Type;
-				AddDescriptor(DescriptorSet, Binding);
-			}
-
-			if (ShaderHeader->InputAttachments.Num())
+			if (ShaderHeader->InputAttachmentInfos.Num())
 			{
 				check(Stage == ShaderStage::Pixel);
-				for (int32 SrcIndex = 0; SrcIndex < ShaderHeader->InputAttachments.Num(); ++SrcIndex)
-				{
-					int32 OriginalGlobalIndex = ShaderHeader->InputAttachments[SrcIndex].GlobalIndex;
-					const FVulkanShaderHeader::FGlobalInfo& OriginalGlobalInfo = ShaderHeader->Globals[OriginalGlobalIndex];
-					check((VkDescriptorType)OriginalGlobalInfo.TypeIndex == VK_DESCRIPTOR_TYPE_INPUT_ATTACHMENT);
-					int32 RemappingIndex = RemappingInfo.InputAttachmentData.AddDefaulted();
-					FInputAttachmentData& AttachmentData = RemappingInfo.InputAttachmentData[RemappingIndex];
-					AttachmentData.BindingIndex = RemappingInfo.StageInfos[Stage].Globals[OriginalGlobalIndex].NewBindingIndex;
-					AttachmentData.DescriptorSet = (uint8)DescriptorSet;
-					AttachmentData.Type = ShaderHeader->InputAttachments[SrcIndex].Type;
-				}
+				check(RemappingInfo.InputAttachmentData.Num() == 0);
+				RemappingInfo.InputAttachmentData = ShaderHeader->InputAttachmentInfos;
 			}
 		}
 	}
@@ -891,10 +813,6 @@ void FVulkanComputePipelineDescriptorInfo::Initialize(const FDescriptorSetRemapp
 
 	const FDescriptorSetRemappingInfo::FStageInfo& StageInfo = InRemappingInfo.StageInfos[ShaderStage::Compute];
 
-	RemappingGlobalInfos = StageInfo.Globals;
-	RemappingUBInfos = StageInfo.UniformBuffers;
-	RemappingPackedUBInfos = StageInfo.PackedUBBindingIndices;
-
 	RemappingInfo = &InRemappingInfo;
 
 	if (StageInfo.Types.Num() > 0)
@@ -913,11 +831,6 @@ void FVulkanGfxPipelineDescriptorInfo::Initialize(const FDescriptorSetRemappingI
 	{
 		const FDescriptorSetRemappingInfo::FStageInfo& StageInfo = InRemappingInfo.StageInfos[StageIndex];
 
-		//#todo-rco: Enable this!
-		RemappingUBInfos[StageIndex] = StageInfo.UniformBuffers;
-		RemappingGlobalInfos[StageIndex] = StageInfo.Globals;
-		RemappingPackedUBInfos[StageIndex] = StageInfo.PackedUBBindingIndices;
-
 		if (StageInfo.Types.Num() > 0)
 		{
 			check(StageIndex < sizeof(HasDescriptorsInSetMask) * 8);
diff --git a/Engine/Source/Runtime/VulkanRHI/Public/VulkanShaderResources.h b/Engine/Source/Runtime/VulkanRHI/Public/VulkanShaderResources.h
--- a/Engine/Source/Runtime/VulkanRHI/Public/VulkanShaderResources.h
+++ b/Engine/Source/Runtime/VulkanRHI/Public/VulkanShaderResources.h
@@ -18,50 +18,30 @@
 // Size = Ignored for non-globals
 struct FVulkanShaderHeader
 {
-	struct FSpirvInfo
+	// Includes all bindings, the index in this array is the binding slot
+	struct FBindingInfo
 	{
-		FSpirvInfo() = default;
-		FSpirvInfo(uint32 InDescriptorSetOffset, uint32 InBindingIndexOffset)
-			: DescriptorSetOffset(InDescriptorSetOffset)
-			, BindingIndexOffset(InBindingIndexOffset)
-		{
-		}
-
-		uint32	DescriptorSetOffset = UINT32_MAX;
-		uint32	BindingIndexOffset = UINT32_MAX;
-	};
-
-	struct FUniformBufferInfo
-	{
-		uint32					LayoutHash;
-		uint16					ConstantDataOriginalBindingIndex;
-		uint8					bOnlyHasResources;
-		uint8					bHasResources;
+		// VkDescriptorType
+		uint32					DescriptorType;
 #if VULKAN_ENABLE_BINDING_DEBUG_NAMES
 		FString					DebugName;
 #endif
 	};
-	TArray<FUniformBufferInfo>	UniformBuffers;
-
-	struct FGlobalInfo
-	{
-		uint32							TypeIndex;
-#if VULKAN_ENABLE_BINDING_DEBUG_NAMES
-		FString							DebugName;
-#endif
-	};
-	TArray<FGlobalInfo>					Globals;
+	TArray<FBindingInfo>		Bindings;
 
-	struct FPackedUBInfo
+	// FBindingInfo with type VK_DESCRIPTOR_TYPE_UNIFORM_BUFFER have a corresponding entry in this table (at the same index)
+	struct FUniformBufferInfo
 	{
-		uint32							SizeInBytes;
-		uint32							SPIRVDescriptorSetOffset;
-		uint32							SPIRVBindingIndexOffset;
+		uint32					LayoutHash;
+		uint8					bHasResources;
+		uint8					BindlessCBIndex;
 	};
-	TArray<FPackedUBInfo>					PackedUBs;
+	TArray<FUniformBufferInfo>	UniformBufferInfos;
 
+	// The order of this enum should always match the strings in VulkanBackend.cpp (VULKAN_SUBPASS_FETCH)
 	enum class EAttachmentType : uint8
 	{
+		Depth,
 		Color0,
 		Color1,
 		Color2,
@@ -70,28 +50,39 @@ struct FVulkanShaderHeader
 		Color5,
 		Color6,
 		Color7,
-		Depth,
 
 		Count,
 	};
-	struct FInputAttachment
+
+	// Used to determine the EAttachmentType of a FBindingInfo with type VK_DESCRIPTOR_TYPE_INPUT_ATTACHMENT
+	struct FInputAttachmentInfo
 	{
-		uint16			GlobalIndex;
-		EAttachmentType	Type;
-		uint8			Pad = 0;
+		uint8					BindingIndex;
+		EAttachmentType			Type;
 	};
-	TArray<FInputAttachment>				InputAttachments;
+	TArray<FInputAttachmentInfo> InputAttachmentInfos;
+
+	// The number of uniform buffers containing constants and requiring bindings
+	// Uniform buffers beyond this index do not have bindings (resource only UB)
+	uint32						NumBoundUniformBuffers = 0;
+
+	// Size of the uniform buffer containing packed globals
+	// If present (not zero), it will always be at binding 0 of the stage
+	uint32						PackedGlobalsSize = 0;
+
+	// Mask of input attachments being used (the index of the bit corresponds to EAttachmentType value)
+	uint32						InputAttachmentsMask = 0;
 
 	// Mostly relevant for Vertex Shaders
-	uint32									InOutMask;
+	uint32						InOutMask;
 
 	// Relevant for Ray Tracing Shaders
-	uint32                                  RayTracingPayloadType = 0;
-	uint32                                  RayTracingPayloadSize = 0;
+	uint32						RayTracingPayloadType = 0;
+	uint32						RayTracingPayloadSize = 0;
 
-	FSHAHash								SourceHash;
-	uint32									SpirvCRC = 0;
-	uint8									WaveSize = 0;
+	FSHAHash					SourceHash;
+	uint32						SpirvCRC = 0;
+	uint8						WaveSize = 0;
 
 	// For RayHitGroup shaders
 	enum class ERayHitGroupEntrypoint : uint8
@@ -107,13 +98,10 @@ struct FVulkanShaderHeader
 		// to circumvent DXC compilation issues
 		SeparateBlob
 	};
-	ERayHitGroupEntrypoint RayGroupAnyHit = ERayHitGroupEntrypoint::NotPresent;
-	ERayHitGroupEntrypoint RayGroupIntersection = ERayHitGroupEntrypoint::NotPresent;
+	ERayHitGroupEntrypoint				RayGroupAnyHit = ERayHitGroupEntrypoint::NotPresent;
+	ERayHitGroupEntrypoint				RayGroupIntersection = ERayHitGroupEntrypoint::NotPresent;
 
-	TArray<FSpirvInfo>						UniformBufferSpirvInfos;
-	TArray<FSpirvInfo>						GlobalSpirvInfos;
-
-	FString									DebugName;
+	FString								DebugName;
 
 	FVulkanShaderHeader() = default;
 	enum EInit
@@ -126,55 +114,38 @@ struct FVulkanShaderHeader
 	}
 };
 
-inline FArchive& operator<<(FArchive& Ar, FVulkanShaderHeader::FSpirvInfo& Info)
+inline FArchive& operator<<(FArchive& Ar, FVulkanShaderHeader::FBindingInfo& BindingInfo)
 {
-	Ar << Info.DescriptorSetOffset;
-	Ar << Info.BindingIndexOffset;
-	return Ar;
-}
-
-inline FArchive& operator<<(FArchive& Ar, FVulkanShaderHeader::FUniformBufferInfo& UBInfo)
-{
-	Ar << UBInfo.LayoutHash;
-	Ar << UBInfo.ConstantDataOriginalBindingIndex;
-	Ar << UBInfo.bOnlyHasResources;
-	Ar << UBInfo.bHasResources;
+	Ar << BindingInfo.DescriptorType;
 #if VULKAN_ENABLE_BINDING_DEBUG_NAMES
-	Ar << UBInfo.DebugName;
+	Ar << BindingInfo.DebugName;
 #endif
 	return Ar;
 }
 
-inline FArchive& operator<<(FArchive& Ar, FVulkanShaderHeader::FPackedUBInfo& PackedUBInfo)
+inline FArchive& operator<<(FArchive& Ar, FVulkanShaderHeader::FUniformBufferInfo& Info)
 {
-	Ar << PackedUBInfo.SizeInBytes;
-	Ar << PackedUBInfo.SPIRVDescriptorSetOffset;
-	Ar << PackedUBInfo.SPIRVBindingIndexOffset;
-	return Ar;
-}
-
-inline FArchive& operator<<(FArchive& Ar, FVulkanShaderHeader::FGlobalInfo& GlobalInfo)
-{
-	Ar << GlobalInfo.TypeIndex;
-#if VULKAN_ENABLE_BINDING_DEBUG_NAMES
-	Ar << GlobalInfo.DebugName;
-#endif
+	Ar << Info.LayoutHash;
+	Ar << Info.bHasResources;
+	Ar << Info.BindlessCBIndex;
 	return Ar;
 }
 
-inline FArchive& operator<<(FArchive& Ar, FVulkanShaderHeader::FInputAttachment& AttachmentInfo)
+inline FArchive& operator<<(FArchive& Ar, FVulkanShaderHeader::FInputAttachmentInfo& Info)
 {
-	Ar << AttachmentInfo.GlobalIndex;
-	Ar << AttachmentInfo.Type;
+	Ar << Info.BindingIndex;
+	Ar << Info.Type;
 	return Ar;
 }
 
 inline FArchive& operator<<(FArchive& Ar, FVulkanShaderHeader& Header)
 {
-	Ar << Header.UniformBuffers;
-	Ar << Header.Globals;
-	Ar << Header.PackedUBs;
-	Ar << Header.InputAttachments;
+	Ar << Header.Bindings;
+	Ar << Header.UniformBufferInfos;
+	Ar << Header.InputAttachmentInfos;
+	Ar << Header.NumBoundUniformBuffers;
+	Ar << Header.PackedGlobalsSize;
+	Ar << Header.InputAttachmentsMask;
 	Ar << Header.InOutMask;
 	Ar << Header.RayTracingPayloadType;
 	Ar << Header.RayTracingPayloadSize;
@@ -183,8 +154,6 @@ inline FArchive& operator<<(FArchive& Ar, FVulkanShaderHeader& Header)
 	Ar << Header.WaveSize;
 	Ar << Header.RayGroupAnyHit;
 	Ar << Header.RayGroupIntersection;
-	Ar << Header.UniformBufferSpirvInfos;
-	Ar << Header.GlobalSpirvInfos;
 	Ar << Header.DebugName;
 	return Ar;
 }
