VulkanRHI: adjust assert until LooseParameterData is processed.

#rnx

diff --git a/Engine/Source/Developer/VulkanShaderFormat/Private/VulkanShaderCompiler.cpp b/Engine/Source/Developer/VulkanShaderFormat/Private/VulkanShaderCompiler.cpp
--- a/Engine/Source/Developer/VulkanShaderFormat/Private/VulkanShaderCompiler.cpp
+++ b/Engine/Source/Developer/VulkanShaderFormat/Private/VulkanShaderCompiler.cpp
@@ -1585,7 +1585,7 @@ static TArray<FString> ConvertUBToBindless(FString& PreprocessedShaderSource)
 
 static void UpdateBindlessUBs(const FVulkanShaderCompilerInternalState& InternalState, VulkanShaderCompilerSerializedOutput& SerializedOutput, FShaderCompilerOutput& Output)
 {
-	check(SerializedOutput.Header.Bindings.Num() == 0);
+	check(SerializedOutput.Header.Bindings.Num() <= 1);  // :todo-jn: allow for globals in slot 0 (looseparams)
 	for (int32 CBIndex = 0; CBIndex < InternalState.AllBindlessUBs.Num(); CBIndex++)
 	{
 		const FString& CBName = InternalState.AllBindlessUBs[CBIndex];
