VulkanRHI: Add LooseParameterData to shader record for Vulkan ray tracing.  Move MaxShaderGroupStride to a Cvar.  Fix record out-of-bounds check to use the stride instead of the full size.

#rnx

diff --git a/Engine/Shaders/Private/HeterogeneousVolumes/HeterogeneousVolumesHardwareRayTracing.usf b/Engine/Shaders/Private/HeterogeneousVolumes/HeterogeneousVolumesHardwareRayTracing.usf
--- a/Engine/Shaders/Private/HeterogeneousVolumes/HeterogeneousVolumesHardwareRayTracing.usf
+++ b/Engine/Shaders/Private/HeterogeneousVolumes/HeterogeneousVolumesHardwareRayTracing.usf
@@ -193,6 +193,10 @@ RAY_TRACING_ENTRY_MISS(SparseVoxelsMissShader,
 	Payload.SetMiss();
 }
 
+// These variables are only used by the ray gen shaders below
+// Investigate why MipLevel isn't dead stripped when building the shaders above (see UE-218402).
+#if RAYGENSHADER
+
 // Scene
 RaytracingAccelerationStructure TLAS;
 
@@ -206,6 +210,8 @@ int MipLevel;
 // Output
 RWTexture2D<float4> RWLightingTexture;
 
+#endif
+
 FRayDesc CreateVolumeRay(
 	float3 Origin,
 	float3 Direction,
diff --git a/Engine/Shaders/Public/Platform/Vulkan/VulkanCommon.ush b/Engine/Shaders/Public/Platform/Vulkan/VulkanCommon.ush
--- a/Engine/Shaders/Public/Platform/Vulkan/VulkanCommon.ush
+++ b/Engine/Shaders/Public/Platform/Vulkan/VulkanCommon.ush
@@ -7,7 +7,7 @@
 #pragma once
 
 // Update this GUID to invalidate shader recompilation for only Vulkan shaders
-#pragma message("UESHADERMETADATA_VERSION 52DDF0DD-4378-4376-9CFF-4DCC8271F475")
+#pragma message("UESHADERMETADATA_VERSION 23542B66-F3AF-43B3-9B33-6C550B201903")
 
 #define COMPILER_SUPPORTS_ULONG_TYPES 1
 #define COMPILER_SUPPORTS_UINT64_IMAGE_ATOMICS (1 && COMPILER_SUPPORTS_ULONG_TYPES)
@@ -86,6 +86,11 @@ void ImageInterlockedAddUInt64(RWTexture2D<UlongType> Dest, uint2 Coord, UlongTy
 	#define OVERRIDE_RAY_TRACING_HIT_GROUP_SYSTEM_RESOURCES 1
 	#include "/Engine/Shared/RayTracingBuiltInResources.h"
 
+	struct VulkanShaderRecordLooseParameterData
+	{
+		uint VulkanShaderRecordDummyGlobals;
+	};
+
 	// Combine the common FHitGroupSystemRootConstants with bindless parameters
 	// :todo-jn: must match with VulkanRayTracing.h decl until we create a common header like RayTracingBuiltInResources.h
 	struct FVulkanHitGroupSystemParameters
@@ -95,7 +100,10 @@ void ImageInterlockedAddUInt64(RWTexture2D<UlongType> Dest, uint2 Coord, UlongTy
 		uint BindlessHitGroupSystemIndexBuffer;
 		uint BindlessHitGroupSystemVertexBuffer;
 
-		uint BindlessUniformBuffers[VULKAN_MAX_BINDLESS_UNIFORM_BUFFERS_PER_STAGE];
+		uint BindlessUniformBuffers[32];
+
+		// Globals must be kept at the end of FVulkanHitGroupSystemParameters
+		VulkanShaderRecordLooseParameterData Globals;
 	};
 
 	[[vk::shader_record_ext]] ConstantBuffer<FVulkanHitGroupSystemParameters> VulkanHitGroupSystemParameters;
diff --git a/Engine/Source/Developer/VulkanShaderFormat/Private/VulkanShaderCompiler.cpp b/Engine/Source/Developer/VulkanShaderFormat/Private/VulkanShaderCompiler.cpp
--- a/Engine/Source/Developer/VulkanShaderFormat/Private/VulkanShaderCompiler.cpp
+++ b/Engine/Source/Developer/VulkanShaderFormat/Private/VulkanShaderCompiler.cpp
@@ -229,6 +229,7 @@ struct FVulkanShaderCompilerInternalState
 	FString IntersectionEntry;
 
 	TArray<FString> AllBindlessUBs;
+	uint32 ShaderRecordGlobalsSize = 0;
 
 	// Forwarded calls for convenience
 	inline EShaderFrequency GetShaderFrequency() const
@@ -1447,7 +1448,6 @@ void ModifyVulkanCompilerInput(FShaderCompilerInput& Input)
 	Input.Environment.SetDefine(TEXT("BINDLESS_SRV_ARRAY_PREFIX"), FShaderParameterParser::kBindlessSRVArrayPrefix);
 	Input.Environment.SetDefine(TEXT("BINDLESS_UAV_ARRAY_PREFIX"), FShaderParameterParser::kBindlessUAVArrayPrefix);
 	Input.Environment.SetDefine(TEXT("BINDLESS_SAMPLER_ARRAY_PREFIX"), FShaderParameterParser::kBindlessSamplerArrayPrefix);
-	Input.Environment.SetDefine(TEXT("VULKAN_MAX_BINDLESS_UNIFORM_BUFFERS_PER_STAGE"), VulkanBindless::MaxUniformBuffersPerStage);
 
 	if (IsAndroidShaderFormat(Input.ShaderFormat))
 	{
@@ -1584,9 +1584,67 @@ static TArray<FString> ConvertUBToBindless(FString& PreprocessedShaderSource)
 }
 
 
+// Helper function to know how much space to set aside in the shader record for a global
+static uint32 GetSizeForType(FStringView TypeName, FStringView ArraySize)
+{
+	static TMap<FStringView, uint32> SizeForTypeMap;
+	if (SizeForTypeMap.Num() == 0)
+	{
+		SizeForTypeMap.Add(FStringView(TEXT("uint")),   4);
+		SizeForTypeMap.Add(FStringView(TEXT("uint2")),  8);
+		SizeForTypeMap.Add(FStringView(TEXT("uint4")),  16);
+		SizeForTypeMap.Add(FStringView(TEXT("float")),  4);
+		SizeForTypeMap.Add(FStringView(TEXT("float2")), 8);
+		SizeForTypeMap.Add(FStringView(TEXT("float4")), 16);
+	}
+
+	checkf(ArraySize.Len() == 0, TEXT("Need to add array support!")); // :todo-jn: Add array parsing
+
+	const uint32* TypeSize = SizeForTypeMap.Find(TypeName);
+	checkf(TypeSize, TEXT("Missing type size for %.*s"), TypeName.Len(), TypeName.GetData());
+	return *TypeSize;
+}
+
+
+// :todo-jn: TEMPORARY EXPERIMENT - will eventually move into preprocessing step
+static uint32 ConvertGlobalsToShaderRecord(const FShaderParameterParser& ShaderParameterParser, const TMap<FStringView, FStringView>& ReplacedGlobals, FString& PreprocessedShaderSource, FShaderCompilerOutput& Output)
+{
+	uint32 ShaderRecordGlobalsSize = 0;
+	uint32 ShaderRecordParamCount = 0;
+	FString ShaderRecordGlobalsString;
+
+	for (const auto& ParamDecl : ReplacedGlobals)
+	{
+		ShaderRecordGlobalsString += ParamDecl.Value;
+
+		const FString ParamName(ParamDecl.Key);
+		const FShaderParameterParser::FParsedShaderParameter& Info = ShaderParameterParser.FindParameterInfos(ParamName);
+		const uint32 ParamSize = GetSizeForType(Info.ParsedType, Info.ParsedArraySize);
+
+		HandleReflectedGlobalConstantBufferMember(
+			ParamName,
+			ShaderRecordParamCount++,
+			ShaderRecordGlobalsSize,
+			ParamSize,
+			Output
+		);
+
+		ShaderRecordGlobalsSize += ParamSize; 
+	}
+
+	if (ShaderRecordGlobalsString.Len())
+	{
+		const int32 ReplacementCount = PreprocessedShaderSource.ReplaceInline(TEXT("uint VulkanShaderRecordDummyGlobals;"), *ShaderRecordGlobalsString, ESearchCase::CaseSensitive);
+		checkf(ReplacementCount == 1, TEXT("VulkanShaderRecordDummyGlobals was replaced %d times!"), ReplacementCount);
+	}
+
+	return ShaderRecordGlobalsSize;
+}
+
+
 static void UpdateBindlessUBs(const FVulkanShaderCompilerInternalState& InternalState, VulkanShaderCompilerSerializedOutput& SerializedOutput, FShaderCompilerOutput& Output)
 {
-	check(SerializedOutput.Header.Bindings.Num() <= 1);  // :todo-jn: allow for globals in slot 0 (looseparams)
+	checkf(SerializedOutput.Header.Bindings.Num() == 0, TEXT("Shaders using bindless UBs should have no other bindings."));
 	for (int32 CBIndex = 0; CBIndex < InternalState.AllBindlessUBs.Num(); CBIndex++)
 	{
 		const FString& CBName = InternalState.AllBindlessUBs[CBIndex];
@@ -1713,11 +1771,12 @@ static bool CompileShaderGroup(
 
 struct FVulkanShaderParameterParserPlatformConfiguration : public FShaderParameterParser::FPlatformConfiguration
 {
-	FVulkanShaderParameterParserPlatformConfiguration(const FShaderCompilerInput& Input)
+	FVulkanShaderParameterParserPlatformConfiguration(const FShaderCompilerInput& Input, TMap<FStringView, FStringView>& InReplacedGlobals)
 		: FShaderParameterParser::FPlatformConfiguration()
 		, bIsRayTracingShader(Input.IsRayTracingShader())
 		, HitGroupSystemIndexBufferName(FShaderParameterParser::kBindlessSRVPrefix + FString(TEXT("HitGroupSystemIndexBuffer")))
 		, HitGroupSystemVertexBufferName(FShaderParameterParser::kBindlessSRVPrefix + FString(TEXT("HitGroupSystemVertexBuffer")))
+		, ReplacedGlobals(InReplacedGlobals)
 	{
 		EnumAddFlags(Flags, EShaderParameterParserConfigurationFlags::SupportsBindless | EShaderParameterParserConfigurationFlags::BindlessUsesArrays);
 
@@ -1727,6 +1786,12 @@ struct FVulkanShaderParameterParserPlatformConfiguration : public FShaderParamet
 			ConstantBufferType = TEXTVIEW("cbuffer");
 			EnumAddFlags(Flags, EShaderParameterParserConfigurationFlags::UseStableConstantBuffer);
 		}
+
+		// Place loose data params in the shader record for shaders with bindless UBs
+		if (bIsRayTracingShader && (Input.Target.GetFrequency() != SF_RayGen))
+		{
+			EnumAddFlags(Flags, EShaderParameterParserConfigurationFlags::ReplaceGlobals);
+		}
 	}
 
 	virtual FString GenerateBindlessAccess(EBindlessConversionType BindlessType, FStringView FullTypeString, FStringView ArrayNameOverride, FStringView IndexString) const final
@@ -1750,9 +1815,21 @@ struct FVulkanShaderParameterParserPlatformConfiguration : public FShaderParamet
 			IndexString.Len(), IndexString.GetData());
 	}
 
+	// Fill the global with the value stored in the shader record
+	virtual FString ReplaceGlobal(FStringView FullDeclString, FStringView ParamName) const final
+	{
+		ReplacedGlobals.Add(ParamName, FullDeclString);
+
+		FString NewDecl(FullDeclString);
+		NewDecl = TEXT("static ") + NewDecl;
+		NewDecl.InsertAt(NewDecl.Find(TEXT(";")), FString::Printf(TEXT(" = VulkanHitGroupSystemParameters.Globals.%.*s"), ParamName.Len(), ParamName.GetData()));
+		return NewDecl;
+	}
+
 	const bool bIsRayTracingShader;
 	const FString HitGroupSystemIndexBufferName;
 	const FString HitGroupSystemVertexBufferName;
+	TMap<FStringView, FStringView>& ReplacedGlobals;
 };
 
 void CompileVulkanShader(const FShaderCompilerInput& Input, const FShaderPreprocessOutput& InPreprocessOutput, FShaderCompilerOutput& Output, const class FString& WorkingDirectory)
@@ -1762,7 +1839,8 @@ void CompileVulkanShader(const FShaderCompilerInput& Input, const FShaderPreproc
 	FString EntryPointName = Input.EntryPointName;
 	FString PreprocessedSource(InPreprocessOutput.GetSourceViewWide());
 
-	FVulkanShaderParameterParserPlatformConfiguration PlatformConfiguration(Input);
+	TMap<FStringView, FStringView> ReplacedGlobals; // Note: these FStringView point to memory in FShaderParameterParser
+	FVulkanShaderParameterParserPlatformConfiguration PlatformConfiguration(Input, ReplacedGlobals);
 	FShaderParameterParser ShaderParameterParser(PlatformConfiguration);
 	if (!ShaderParameterParser.ParseAndModify(Input, Output.Errors, PreprocessedSource))
 	{
@@ -1774,10 +1852,11 @@ void CompileVulkanShader(const FShaderCompilerInput& Input, const FShaderPreproc
 
 	if (InternalState.bUseBindlessUniformBuffer)
 	{
+		InternalState.ShaderRecordGlobalsSize = ConvertGlobalsToShaderRecord(ShaderParameterParser, ReplacedGlobals, PreprocessedSource, Output);
 		InternalState.AllBindlessUBs = ConvertUBToBindless(PreprocessedSource);
 	}
 
-	if (ShaderParameterParser.DidModifyShader() || InternalState.AllBindlessUBs.Num() > 0)
+	if (ShaderParameterParser.DidModifyShader() || InternalState.AllBindlessUBs.Num() || InternalState.ShaderRecordGlobalsSize)
 	{
 		Output.ModifiedShaderSource = PreprocessedSource;
 	}
diff --git a/Engine/Source/Runtime/RHI/Public/RHIGlobals.h b/Engine/Source/Runtime/RHI/Public/RHIGlobals.h
--- a/Engine/Source/Runtime/RHI/Public/RHIGlobals.h
+++ b/Engine/Source/Runtime/RHI/Public/RHIGlobals.h
@@ -434,6 +434,9 @@ struct FRHIGlobals
 		/** Whether or not the RHI supports inline ray tracing in compute shaders, without a full ray tracing pipeline. */
 		bool SupportsInlineRayTracing = false;
 
+		/** Wether an extra uniform buffer parameter is required when loose parameters are present, or if they are stored directly in the shader record. */
+		bool SupportsLooseParamsInShaderRecord = false;
+
 		/** Required alignment for ray tracing acceleration structures. */
 		uint32 AccelerationStructureAlignment = 0;
 
diff --git a/Engine/Source/Runtime/RenderCore/Private/ShaderParameterParser.cpp b/Engine/Source/Runtime/RenderCore/Private/ShaderParameterParser.cpp
--- a/Engine/Source/Runtime/RenderCore/Private/ShaderParameterParser.cpp
+++ b/Engine/Source/Runtime/RenderCore/Private/ShaderParameterParser.cpp
@@ -988,9 +988,11 @@ void FShaderParameterParser::ApplyBindlessModifications(FString& PreprocessedSha
 		TArray<FShaderCodeModifications> Modifications;
 		Modifications.Reserve(ParsedParameters.Num());
 
+		const bool bReplaceGlobals = EnumHasAnyFlags(PlatformConfiguration.Flags, EShaderParameterParserConfigurationFlags::ReplaceGlobals);
+
 		for (TPair<FString, FParsedShaderParameter>& Itr : ParsedParameters)
 		{
-			FParsedShaderParameter& ParsedParameter = Itr.Value;
+			const FParsedShaderParameter& ParsedParameter = Itr.Value;
 
 			if (!ParsedParameter.IsFound())
 			{
@@ -1006,6 +1008,26 @@ void FShaderParameterParser::ApplyBindlessModifications(FString& PreprocessedSha
 
 				Modifications.Add(Modif);
 			}
+			else if (bReplaceGlobals)
+			{
+				const bool IsGlobalParam = 
+					(ParsedParameter.BaseType == UBMT_INVALID) &&
+					!ParsedParameter.ParsedName.StartsWith(FShaderParameterParser::kBindlessSamplerArrayPrefix) &&
+					!ParsedParameter.ParsedName.StartsWith(FShaderParameterParser::kBindlessSRVArrayPrefix) &&
+					!ParsedParameter.ParsedName.StartsWith(FShaderParameterParser::kBindlessUAVArrayPrefix);
+
+				if (IsGlobalParam)
+				{
+					FShaderCodeModifications Modif;
+					Modif.CharOffsetStart = ParsedParameter.ParsedCharOffsetStart;
+					Modif.CharOffsetEnd = ParsedParameter.ParsedCharOffsetEnd + 1;
+
+					const int32 NumChars = Modif.CharOffsetEnd - Modif.CharOffsetStart;
+					Modif.Replace = PlatformConfiguration.ReplaceGlobal(FStringView(&OriginalParsedShader[Modif.CharOffsetStart], NumChars), ParsedParameter.ParsedName);
+
+					Modifications.Add(Modif);
+				}
+			}
 		}
 
 		// Apply all modifications
diff --git a/Engine/Source/Runtime/RenderCore/Public/ShaderParameterParser.h b/Engine/Source/Runtime/RenderCore/Public/ShaderParameterParser.h
--- a/Engine/Source/Runtime/RenderCore/Public/ShaderParameterParser.h
+++ b/Engine/Source/Runtime/RenderCore/Public/ShaderParameterParser.h
@@ -31,6 +31,8 @@ enum class EShaderParameterParserConfigurationFlags
 	BindlessUsesArrays      = 1 << 2,
 	// Some RHI Validation requires parameters to be parsed
 	AlwaysParseParams		= 1 << 3,
+	// Replace loose data globals in bindless shaders
+	ReplaceGlobals			= 1 << 4,
 };
 ENUM_CLASS_FLAGS(EShaderParameterParserConfigurationFlags)
 
@@ -61,6 +63,13 @@ public:
 			return FString();
 		}
 
+		/** Generate shader code for accessing a bindless resource or sampler */
+		virtual FString ReplaceGlobal(FStringView FullDeclString, FStringView ParamName) const
+		{
+			checkf(false, TEXT("Platforms that support replacing globals must override ReplaceGlobal"));
+			return FString();
+		}
+
 		TConstArrayView<FStringView> ExtraSRVTypes;
 		TConstArrayView<FStringView> ExtraUAVTypes;
 		FStringView ConstantBufferType;
@@ -94,7 +103,7 @@ public:
 		int32 ParsedPragmaLineOffset = 0;
 		int32 ParsedLineOffset = 0;
 
-		/** Character position of the start and end of the parameter decelaration in FParsedShaderParameter::OriginalParsedShader */
+		/** Character position of the start and end of the parameter declaration in FParsedShaderParameter::OriginalParsedShader */
 		int32 ParsedCharOffsetStart = INDEX_NONE;
 		int32 ParsedCharOffsetEnd = INDEX_NONE;
 
diff --git a/Engine/Source/Runtime/Renderer/Private/MeshPassProcessor.cpp b/Engine/Source/Runtime/Renderer/Private/MeshPassProcessor.cpp
--- a/Engine/Source/Runtime/Renderer/Private/MeshPassProcessor.cpp
+++ b/Engine/Source/Runtime/Renderer/Private/MeshPassProcessor.cpp
@@ -348,7 +348,8 @@ FRayTracingLocalShaderBindings* FMeshDrawShaderBindings::SetRayTracingShaderBind
 		}
 	}
 
-	checkf(MaxUniformBufferIndex + 1 == NumUniformBufferParameters + LooseParameterBuffers.Num(),
+	// Vulkan places loose parameter data in the shader record which does not add to UB count
+	checkf(MaxUniformBufferIndex + 1 == NumUniformBufferParameters + (GRHIGlobals.RayTracing.SupportsLooseParamsInShaderRecord ? 0 : LooseParameterBuffers.Num()),
 		TEXT("Highest index of a uniform buffer was %d, but there were %d uniform buffer parameters and %d loose parameters"),
 		MaxUniformBufferIndex,
 		NumUniformBufferParameters,
diff --git a/Engine/Source/Runtime/VulkanRHI/Private/VulkanRHI.cpp b/Engine/Source/Runtime/VulkanRHI/Private/VulkanRHI.cpp
--- a/Engine/Source/Runtime/VulkanRHI/Private/VulkanRHI.cpp
+++ b/Engine/Source/Runtime/VulkanRHI/Private/VulkanRHI.cpp
@@ -937,6 +937,9 @@ void FVulkanDynamicRHI::InitInstance()
 													Device->GetOptionalExtensionProperties().AccelerationStructureProps.minAccelerationStructureScratchOffsetAlignment);
 
 			GRHIRayTracingInstanceDescriptorSize = uint32(sizeof(VkAccelerationStructureInstanceKHR));
+
+			// Loose parameters are always placed in the shader record after the FVulkanHitGroupSystemParameters in Vulkan (see VulkanRayTracing.h and VulkanCommon.ush)
+			GRHIGlobals.RayTracing.SupportsLooseParamsInShaderRecord = true;
 		}
 #if VULKAN_ENABLE_DUMP_LAYER
 		// Disable RHI thread by default if the dump layer is enabled
diff --git a/Engine/Source/Runtime/VulkanRHI/Private/VulkanRayTracing.cpp b/Engine/Source/Runtime/VulkanRHI/Private/VulkanRayTracing.cpp
--- a/Engine/Source/Runtime/VulkanRHI/Private/VulkanRayTracing.cpp
+++ b/Engine/Source/Runtime/VulkanRHI/Private/VulkanRayTracing.cpp
@@ -46,6 +46,13 @@ static FAutoConsoleVariableRef GCVarSubmitOnTraceRays(
 	ECVF_ReadOnly
 );
 
+static int32 GVulkanRayTracingMaxShaderGroupStride = 4096;
+static FAutoConsoleVariableRef GCVarVulkanRayTracingMaxShaderGroupStride(
+	TEXT("r.Vulkan.RayTracing.MaxShaderGroupStride"),
+	GVulkanRayTracingMaxShaderGroupStride,
+	TEXT("The default size to allocate for each record (default: 4096)."),
+	ECVF_ReadOnly
+);
 
 // Ray tracing stat counters
 
@@ -928,15 +935,14 @@ FVulkanRayTracingShaderTable::FVulkanRayTracingShaderTable(FVulkanDevice* Device
 	, HandleSize(Device->GetOptionalExtensionProperties().RayTracingPipelineProps.shaderGroupHandleSize)
 	, HandleSizeAligned(Align(HandleSize, Device->GetOptionalExtensionProperties().RayTracingPipelineProps.shaderGroupHandleAlignment))
 {
-	auto InitAlloc = [Device = Device, HandleSizeAligned = HandleSizeAligned](FVulkanShaderTableAllocation& Alloc, uint32 InHandleCount, bool InUseLocalRecord) 
+	auto InitAlloc = [HandleSizeAligned = HandleSizeAligned](FVulkanShaderTableAllocation& Alloc, uint32 InHandleCount, bool InUseLocalRecord) 
 	{
 		Alloc.HandleCount = InHandleCount;
 		Alloc.bUseLocalRecord = InUseLocalRecord;
 
 		if (Alloc.HandleCount > 0)
 		{
-			const VkPhysicalDeviceRayTracingPipelinePropertiesKHR& RayTracingPipelineProps = Device->GetOptionalExtensionProperties().RayTracingPipelineProps;
-			Alloc.Region.stride = InUseLocalRecord ? FMath::Min<VkDeviceSize>(RayTracingPipelineProps.maxShaderGroupStride, 4096) : HandleSizeAligned; // :todo-jn: shrink stride to necessary amount
+			Alloc.Region.stride = InUseLocalRecord ? (uint32)GVulkanRayTracingMaxShaderGroupStride : HandleSizeAligned;
 			Alloc.Region.size = Alloc.HandleCount * Alloc.Region.stride;
 
 			// Host buffer
@@ -1024,7 +1030,7 @@ void FVulkanRayTracingShaderTable::SetLocalShaderParameters(EShaderFrequency Fre
 
 	checkfSlow(OffsetWithinRecord % 4 == 0, TEXT("SBT record parameters must be written on DWORD-aligned boundary"));
 	checkfSlow(InDataSize % 4 == 0, TEXT("SBT record parameters must be DWORD-aligned"));
-	checkf(OffsetWithinRecord + InDataSize <= Alloc.Region.size, TEXT("SBT record write request is out of bounds"));
+	checkf(OffsetWithinRecord + InDataSize <= Alloc.Region.stride, TEXT("SBT record write request is out of bounds"));
 
 	const uint32 WriteOffset = HandleSizeAligned + (Alloc.Region.stride * RecordIndex) + OffsetWithinRecord;
 	FMemory::Memcpy(&Alloc.HostBuffer[WriteOffset], InData, InDataSize);
@@ -1366,6 +1372,10 @@ static FRHIRayTracingShader* GetBuiltInRayTracingShader()
 
 void FVulkanDevice::InitializeRayTracing()
 {
+	const VkPhysicalDeviceRayTracingPipelinePropertiesKHR& RayTracingPipelineProps = GetOptionalExtensionProperties().RayTracingPipelineProps;
+	checkf((uint32)GVulkanRayTracingMaxShaderGroupStride <= RayTracingPipelineProps.maxShaderGroupStride, 
+		TEXT("Specified value for r.Vulkan.RayTracing.MaxShaderGroupStride is too large for this device! It will be capped."));
+	GVulkanRayTracingMaxShaderGroupStride = FMath::Min<VkDeviceSize>(RayTracingPipelineProps.maxShaderGroupStride, GVulkanRayTracingMaxShaderGroupStride);
 }
 
 // Temporary code to generate dummy UBs to bind when none is provided to prevent bindless code from crashing
@@ -1576,7 +1586,7 @@ FVulkanRayTracingPipelineState::FVulkanRayTracingPipelineState(FVulkanDevice* co
 	RayTracingPipelineCreateInfo.maxPipelineRayRecursionDepth = 1;
 	RayTracingPipelineCreateInfo.layout = InDevice->GetBindlessDescriptorManager()->GetPipelineLayout();
 	RayTracingPipelineCreateInfo.flags = VK_PIPELINE_CREATE_DESCRIPTOR_BUFFER_BIT_EXT;
-	
+
 	VkDeferredOperationKHR DeferredOp = VK_NULL_HANDLE; // :todo-jn: more speed
 	if (GVulkanRayTracingAllowDeferredOperation >= 0)
 	{
@@ -2199,6 +2209,7 @@ void FVulkanCommandListContext::RHIRayTraceDispatchIndirect(
 
 static void SetSystemParametersUB(FVulkanHitGroupSystemParameters& OutSystemParameters, FVulkanRayTracingShaderTable* ShaderTable, uint32 InNumUniformBuffers, FRHIUniformBuffer* const* InUniformBuffers, const FVulkanRayTracingShader* InShader)
 {
+
 	// Plug the shaders in the right slots using LayoutHash comparisons
 	check(InShader->GetCodeHeader().UniformBufferInfos.Num() <= (int32)InNumUniformBuffers);
 	for (int32 UBIndex = 0; UBIndex < InShader->GetCodeHeader().UniformBufferInfos.Num(); ++UBIndex)
@@ -2218,6 +2229,7 @@ static void SetSystemParametersUB(FVulkanHitGroupSystemParameters& OutSystemPara
 
 		const FRHIDescriptorHandle BindlessHandle = UniformBuffer->GetBindlessHandle();
 		check(BindlessHandle.IsValid());
+		check(UniformBufferInfo.BindlessCBIndex < UE_ARRAY_COUNT(OutSystemParameters.BindlessUniformBuffers));
 		OutSystemParameters.BindlessUniformBuffers[UniformBufferInfo.BindlessCBIndex] = BindlessHandle.GetIndex();
 
 		ShaderTable->AddUBRef(UniformBuffer);
@@ -2252,6 +2264,7 @@ static void SetRayTracingHitGroup(
 		SetSystemParametersUB(SystemParameters, ShaderTable, NumUniformBuffers, UniformBuffers, Shader);
 
 		ShaderTable->SetLocalShaderParameters(SF_RayHitGroup, RecordIndex, 0, SystemParameters);
+		ShaderTable->SetLooseParameterData(SF_RayHitGroup, RecordIndex, LooseParameterData, LooseParameterDataSize);
 	}
 
 	ShaderTable->SetSlot(SF_RayHitGroup, RecordIndex, HitGroupIndex, Pipeline->GetShaderHandles(SF_RayHitGroup));
@@ -2262,6 +2275,7 @@ static void SetGenericSystemParameters(
 	FVulkanRayTracingShaderTable* ShaderTable, uint32 RecordIndex,
 	FRHIRayTracingPipelineState* InPipeline, uint32 ShaderIndexInPipeline,
 	uint32 NumUniformBuffers, FRHIUniformBuffer* const* UniformBuffers,
+	uint32 LooseParameterDataSize, const void* LooseParameterData,
 	uint32 UserData, const EShaderFrequency ShaderFrequency)
 {
 	FVulkanRayTracingPipelineState* Pipeline = ResourceCast(InPipeline);
@@ -2272,6 +2286,7 @@ static void SetGenericSystemParameters(
 	SystemParameters.RootConstants.UserData = UserData;
 	SetSystemParametersUB(SystemParameters, ShaderTable, NumUniformBuffers, UniformBuffers, Shader);
 	ShaderTable->SetLocalShaderParameters(ShaderFrequency, RecordIndex, 0, SystemParameters);
+	ShaderTable->SetLooseParameterData(ShaderFrequency, RecordIndex, LooseParameterData, LooseParameterDataSize);
 
 	ShaderTable->SetSlot(ShaderFrequency, RecordIndex, ShaderIndexInPipeline, Pipeline->GetShaderHandles(ShaderFrequency));
 }
@@ -2326,6 +2341,7 @@ void FVulkanCommandListContext::RHISetBindingsOnShaderBindingTable(FRHIShaderBin
 				ShaderTable, Binding.RecordIndex,
 				Pipeline, Binding.ShaderIndexInPipeline,
 				Binding.NumUniformBuffers, Binding.UniformBuffers,
+				Binding.LooseParameterDataSize, Binding.LooseParameterData,
 				Binding.UserData,
 				SF_RayCallable);
 		}
@@ -2335,6 +2351,7 @@ void FVulkanCommandListContext::RHISetBindingsOnShaderBindingTable(FRHIShaderBin
 				ShaderTable, Binding.RecordIndex,
 				Pipeline, Binding.ShaderIndexInPipeline,
 				Binding.NumUniformBuffers, Binding.UniformBuffers,
+				Binding.LooseParameterDataSize, Binding.LooseParameterData,
 				Binding.UserData,
 				SF_RayMiss);
 		}
diff --git a/Engine/Source/Runtime/VulkanRHI/Private/VulkanRayTracing.h b/Engine/Source/Runtime/VulkanRHI/Private/VulkanRayTracing.h
--- a/Engine/Source/Runtime/VulkanRHI/Private/VulkanRayTracing.h
+++ b/Engine/Source/Runtime/VulkanRHI/Private/VulkanRayTracing.h
@@ -25,7 +25,9 @@ struct FVulkanHitGroupSystemParameters
 	uint32 BindlessHitGroupSystemIndexBuffer;
 	uint32 BindlessHitGroupSystemVertexBuffer;
 
-	uint32 BindlessUniformBuffers[16];
+	uint32 BindlessUniformBuffers[32];
+
+	// *** Globals start here ***
 };
 
 
@@ -79,6 +81,16 @@ public:
 
 	void SetLocalShaderParameters(EShaderFrequency Frequency, uint32 RecordIndex, uint32 OffsetWithinRecord, const void* InData, uint32 InDataSize);
 
+	void SetLooseParameterData(EShaderFrequency Frequency, uint32 RecordIndex, const void* InData, uint32 InDataSize)
+	{
+		if (InData && InDataSize)
+		{
+			// Place the loose parameter data after the FVulkanHitGroupSystemParameters in the shader record
+			const uint32 LooseParameterDataOffset = Align(sizeof(FVulkanHitGroupSystemParameters), 4);
+			SetLocalShaderParameters(Frequency, RecordIndex, LooseParameterDataOffset, InData, InDataSize);
+		}
+	}
+
 	void Commit(FVulkanCommandListContext& Context);
 
 	void AddUBRef(FRHIUniformBuffer* UB)
