- Pulled more functions into the SpirvShaderCompiler.inl

diff --git a/Engine/Source/Developer/VulkanShaderFormat/Private/SpirVShaderCompiler.inl b/Engine/Source/Developer/VulkanShaderFormat/Private/SpirVShaderCompiler.inl
--- a/Engine/Source/Developer/VulkanShaderFormat/Private/SpirVShaderCompiler.inl
+++ b/Engine/Source/Developer/VulkanShaderFormat/Private/SpirVShaderCompiler.inl
@@ -1347,4 +1347,247 @@ static void ModifyCompilerInput(FSpirvShaderCompilerInternalState& InternalState
 	}
 }
 
-}; // SpirvShaderCompiler
\ No newline at end of file
+static void UpdateBindlessUBs(const FSpirvShaderCompilerInternalState& InternalState, SpirvShaderCompilerSerializedOutput& SerializedOutput, FShaderCompilerOutput& Output)
+{
+	checkf(SerializedOutput.Header.Bindings.Num() == 0, TEXT("Shaders using bindless UBs should have no other bindings."));
+	for (int32 CBIndex = 0; CBIndex < InternalState.AllBindlessUBs.Num(); CBIndex++)
+	{
+		const FString& CBName = InternalState.AllBindlessUBs[CBIndex];
+
+		// It's possible SPIRV compilation has optimized out a buffer from every shader in the group
+		if (SerializedOutput.UsedBindlessUB.Contains(CBName))
+		{
+			FVulkanShaderHeader::FUniformBufferInfo& Info = SerializedOutput.Header.UniformBufferInfos.AddZeroed_GetRef();
+			Info.LayoutHash = SpirvShaderCompiler::GetUBLayoutHash(InternalState.Input, CBName);
+			Info.BindlessCBIndex = CBIndex;
+
+			const int32 UBIndex = SerializedOutput.Header.UniformBufferInfos.Num() - 1;
+			Output.ParameterMap.AddParameterAllocation(CBName, UBIndex, 0, 1, EShaderParameterType::UniformBuffer);
+		}
+	}
+}
+
+// :todo-jn: TEMPORARY EXPERIMENT - will eventually move into preprocessing step
+static TArray<FString> ConvertUBToBindless(FString& PreprocessedShaderSource)
+{
+	// Fill a map so we pull our bindless sampler/resource indices from the right struct
+	// :todo-jn: Do we not have the layout somewhere instead of calculating offsets?  there must be a better way...
+	auto GenerateNewDecl = [](const int32 CBIndex, const FString& Members, const FString& CBName)
+	{
+		const FString PrefixedCBName = FString::Printf(TEXT("%s%d_%s"), *SpirvShaderCompiler::kBindlessCBPrefix, CBIndex, *CBName);
+		const FString BindlessCBType = PrefixedCBName + TEXT("_Type");
+		const FString BindlessCBHeapName = PrefixedCBName + SpirvShaderCompiler::kBindlessHeapSuffix;
+		const FString PaddingName = FString::Printf(TEXT("%s_Padding"), *CBName);
+
+		FString CBDecl;
+		CBDecl.Reserve(Members.Len() * 3);  // start somewhere approx less bad
+
+		// Declare the struct
+		CBDecl += TEXT("struct ") + BindlessCBType + TEXT(" \n{\n") + Members + TEXT("\n};\n");
+
+		// Declare the safetype and bindless array for this cb
+		CBDecl += FString::Printf(TEXT("ConstantBuffer<%s> %s[];\n"), *BindlessCBType, *BindlessCBHeapName);
+
+		// Now bring in the CB
+		CBDecl += FString::Printf(TEXT("static const %s %s = %s[VulkanHitGroupSystemParameters.BindlessUniformBuffers[%d]];\n"),
+			*BindlessCBType, *PrefixedCBName, *BindlessCBHeapName, CBIndex);
+
+		// Now create a global scope var for each value (as the cbuffer would provide) to patch in seemlessly with the rest of the code
+		uint32 MemberOffset = 0;
+		const TCHAR* MemberSearchPtr = *Members;
+		const uint32 LastMemberSemicolonIndex = Members.Find(TEXT(";"), ESearchCase::CaseSensitive, ESearchDir::FromEnd, -1);
+		check(LastMemberSemicolonIndex != INDEX_NONE);
+		const TCHAR* LastMemberSemicolon = &Members[LastMemberSemicolonIndex];
+
+		do
+		{
+			const TCHAR* MemberTypeStartPtr = nullptr;
+			const TCHAR* MemberTypeEndPtr = nullptr;
+			ParseHLSLTypeName(MemberSearchPtr, MemberTypeStartPtr, MemberTypeEndPtr);
+			const FString MemberTypeName = FString::ConstructFromPtrSize(MemberTypeStartPtr, MemberTypeEndPtr - MemberTypeStartPtr);
+
+			FString MemberName;
+			MemberSearchPtr = ParseHLSLSymbolName(MemberTypeEndPtr, MemberName);
+			check(MemberName.Len() > 0);
+
+			if (MemberName.StartsWith(PaddingName))
+			{
+				while (*MemberSearchPtr && *MemberSearchPtr != ';')
+				{
+					MemberSearchPtr++;
+				}
+			}
+			else
+			{
+				// Skip over trailing tokens and pick up arrays
+				FString ArrayDecl;
+				while (*MemberSearchPtr && *MemberSearchPtr != ';')
+				{
+					if (*MemberSearchPtr == '[')
+					{
+						ArrayDecl.AppendChar(*MemberSearchPtr);
+
+						MemberSearchPtr++;
+						while (*MemberSearchPtr && *MemberSearchPtr != ']')
+						{
+							ArrayDecl.AppendChar(*MemberSearchPtr);
+							MemberSearchPtr++;
+						}
+
+						ArrayDecl.AppendChar(*MemberSearchPtr);
+					}
+
+					MemberSearchPtr++;
+				}
+
+				CBDecl += FString::Printf(TEXT("static const %s %s%s = %s.%s;\n"), *MemberTypeName, *MemberName, *ArrayDecl, *PrefixedCBName, *MemberName);
+			}
+
+			MemberSearchPtr++;
+
+		} while (MemberSearchPtr < LastMemberSemicolon);
+
+		return CBDecl;
+	};
+
+	// replace "cbuffer" decl with a struct filled from bindless constant buffer
+	TArray<FString> BindlessUBs;
+	{
+		const FString UniformBufferDeclIdentifier = TEXT("cbuffer");
+
+		int32 SearchIndex = PreprocessedShaderSource.Find(UniformBufferDeclIdentifier, ESearchCase::CaseSensitive, ESearchDir::FromStart, -1);
+		while (SearchIndex != INDEX_NONE)
+		{
+			FString StructName;
+			const TCHAR* StructNameEndPtr = ParseHLSLSymbolName(&PreprocessedShaderSource[SearchIndex + UniformBufferDeclIdentifier.Len()], StructName);
+			check(StructName.Len() > 0);
+
+			const int32 CBIndex = BindlessUBs.Add(StructName);
+			check(CBIndex < 16);
+
+			const TCHAR* OpeningBracePtr = FCString::Strstr(&PreprocessedShaderSource[SearchIndex + UniformBufferDeclIdentifier.Len()], TEXT("{"));
+			check(OpeningBracePtr);
+			const TCHAR* ClosingBracePtr = FindMatchingClosingBrace(OpeningBracePtr + 1);
+			check(ClosingBracePtr);
+			const int32 ClosingBraceIndex = ClosingBracePtr - (*PreprocessedShaderSource);
+
+			const FString Members = FString::ConstructFromPtrSize(OpeningBracePtr + 1, ClosingBracePtr - OpeningBracePtr - 1);
+			const FString NewDecl = GenerateNewDecl(CBIndex, Members, StructName);
+
+			const int32 OldDeclLen = ClosingBraceIndex - SearchIndex + 1;
+			PreprocessedShaderSource.RemoveAt(SearchIndex, OldDeclLen, EAllowShrinking::No);
+			PreprocessedShaderSource.InsertAt(SearchIndex, NewDecl);
+
+			SearchIndex = PreprocessedShaderSource.Find(UniformBufferDeclIdentifier, ESearchCase::CaseSensitive, ESearchDir::FromStart, SearchIndex + NewDecl.Len());
+		}
+	}
+	return BindlessUBs;
+}
+
+static bool CompileShaderGroup(
+	FSpirvShaderCompilerInternalState& InternalState,
+	const FShaderSource::FStringType& OriginalPreprocessedShaderSource,
+	FShaderCompilerOutput& MergedOutput
+)
+{
+	checkf(InternalState.bSupportsBindless && InternalState.bUseBindlessUniformBuffer, TEXT("Ray tracing requires full bindless in Vulkan."));
+
+	// Compile each one of the shader modules seperately and create one big blob for the engine
+	auto CompilePartialExport = [&OriginalPreprocessedShaderSource, &InternalState, &MergedOutput](
+		FSpirvShaderCompilerInternalState::EHitGroupShaderType HitGroupShaderType,
+		const TCHAR* PartialFileExtension,
+		SpirvShaderCompilerSerializedOutput& PartialSerializedOutput)
+	{
+		InternalState.HitGroupShaderType = HitGroupShaderType;
+
+		FShaderCompilerOutput TempOutput;
+		const bool bIsClosestHit = (HitGroupShaderType == FSpirvShaderCompilerInternalState::EHitGroupShaderType::ClosestHit);
+		FShaderCompilerOutput& PartialOutput = bIsClosestHit ? MergedOutput : TempOutput;
+
+		FShaderSource::FViewType OrigSourceView(OriginalPreprocessedShaderSource);
+		FShaderSource PartialPreprocessedShaderSource(OrigSourceView);
+		UE::ShaderCompilerCommon::RemoveDeadCode(PartialPreprocessedShaderSource, InternalState.GetEntryPointName(), PartialOutput.Errors);
+
+		if (InternalState.bDebugDump)
+		{
+			DumpDebugShaderText(InternalState.Input, PartialPreprocessedShaderSource.GetView().GetData(), *FString::Printf(TEXT("%s.hlsl"), PartialFileExtension));
+		}
+
+		const bool bPartialSuccess = SpirvShaderCompiler::CompileWithShaderConductor(InternalState, PartialPreprocessedShaderSource.GetView(), PartialSerializedOutput, PartialOutput);
+
+		if (!bIsClosestHit)
+		{
+			MergedOutput.NumInstructions = FMath::Max(MergedOutput.NumInstructions, PartialOutput.NumInstructions);
+			MergedOutput.NumTextureSamplers = FMath::Max(MergedOutput.NumTextureSamplers, PartialOutput.NumTextureSamplers);
+			MergedOutput.Errors.Append(MoveTemp(PartialOutput.Errors));
+		}
+
+		return bPartialSuccess;
+	};
+
+	bool bSuccess = false;
+
+	// Closest Hit Module, always present
+	SpirvShaderCompilerSerializedOutput ClosestHitSerializedOutput;
+	{
+		bSuccess = CompilePartialExport(FSpirvShaderCompilerInternalState::EHitGroupShaderType::ClosestHit, TEXT("closest"), ClosestHitSerializedOutput);
+	}
+
+	// Any Hit Module, optional
+	const bool bHasAnyHitModule = !InternalState.AnyHitEntry.IsEmpty();
+	SpirvShaderCompilerSerializedOutput AnyHitSerializedOutput;
+	if (bSuccess && bHasAnyHitModule)
+	{
+		bSuccess = CompilePartialExport(FSpirvShaderCompilerInternalState::EHitGroupShaderType::AnyHit, TEXT("anyhit"), AnyHitSerializedOutput);
+	}
+
+	// Intersection Module, optional
+	const bool bHasIntersectionModule = !InternalState.IntersectionEntry.IsEmpty();
+	SpirvShaderCompilerSerializedOutput IntersectionSerializedOutput;
+	if (bSuccess && bHasIntersectionModule)
+	{
+		bSuccess = CompilePartialExport(FSpirvShaderCompilerInternalState::EHitGroupShaderType::Intersection, TEXT("intersection"), IntersectionSerializedOutput);
+	}
+
+	// Collapse the bindless UB usage into one set and then update the headers
+	ClosestHitSerializedOutput.UsedBindlessUB.Append(AnyHitSerializedOutput.UsedBindlessUB);
+	ClosestHitSerializedOutput.UsedBindlessUB.Append(IntersectionSerializedOutput.UsedBindlessUB);
+	UpdateBindlessUBs(InternalState, ClosestHitSerializedOutput, MergedOutput);
+
+	{
+		// :todo-jn: Having multiple entrypoints in a single SPIRV blob crashes on FLumenHardwareRayTracingMaterialHitGroup for some reason
+		// Adjust the header before we write it out
+		ClosestHitSerializedOutput.Header.RayGroupAnyHit = bHasAnyHitModule ? FVulkanShaderHeader::ERayHitGroupEntrypoint::SeparateBlob : FVulkanShaderHeader::ERayHitGroupEntrypoint::NotPresent;
+		ClosestHitSerializedOutput.Header.RayGroupIntersection = bHasIntersectionModule ? FVulkanShaderHeader::ERayHitGroupEntrypoint::SeparateBlob : FVulkanShaderHeader::ERayHitGroupEntrypoint::NotPresent;
+
+		check(ClosestHitSerializedOutput.Spirv.Data.Num() != 0);
+		FMemoryWriter Ar(MergedOutput.ShaderCode.GetWriteAccess(), true);
+		Ar << ClosestHitSerializedOutput.Header;
+		Ar << ClosestHitSerializedOutput.ShaderResourceTable;
+
+		{
+			uint32 SpirvCodeSizeBytes = ClosestHitSerializedOutput.Spirv.GetByteSize();
+			Ar << SpirvCodeSizeBytes;
+			Ar.Serialize((uint8*)ClosestHitSerializedOutput.Spirv.Data.GetData(), SpirvCodeSizeBytes);
+		}
+
+		if (bHasAnyHitModule)
+		{
+			uint32 SpirvCodeSizeBytes = AnyHitSerializedOutput.Spirv.GetByteSize();
+			Ar << SpirvCodeSizeBytes;
+			Ar.Serialize((uint8*)AnyHitSerializedOutput.Spirv.Data.GetData(), SpirvCodeSizeBytes);
+		}
+
+		if (bHasIntersectionModule)
+		{
+			uint32 SpirvCodeSizeBytes = IntersectionSerializedOutput.Spirv.GetByteSize();
+			Ar << SpirvCodeSizeBytes;
+			Ar.Serialize((uint8*)IntersectionSerializedOutput.Spirv.Data.GetData(), SpirvCodeSizeBytes);
+		}
+	}
+
+	MergedOutput.bSucceeded = bSuccess;
+	return bSuccess;
+}
+
+}; // SpirvShaderCompiler
diff --git a/Engine/Source/Developer/VulkanShaderFormat/Private/VulkanShaderCompiler.cpp b/Engine/Source/Developer/VulkanShaderFormat/Private/VulkanShaderCompiler.cpp
--- a/Engine/Source/Developer/VulkanShaderFormat/Private/VulkanShaderCompiler.cpp
+++ b/Engine/Source/Developer/VulkanShaderFormat/Private/VulkanShaderCompiler.cpp
@@ -152,124 +152,6 @@ void ModifyVulkanCompilerInput(FShaderCompilerInput& Input)
 	SpirvShaderCompiler::ModifyCompilerInput(InternalState, Input);
 }
 
-// :todo-jn: TEMPORARY EXPERIMENT - will eventually move into preprocessing step
-static TArray<FString> ConvertUBToBindless(FString& PreprocessedShaderSource)
-{
-	// Fill a map so we pull our bindless sampler/resource indices from the right struct
-	// :todo-jn: Do we not have the layout somewhere instead of calculating offsets?  there must be a better way...
-	auto GenerateNewDecl = [](const int32 CBIndex, const FString& Members, const FString& CBName)
-	{
-		const FString PrefixedCBName = FString::Printf(TEXT("%s%d_%s"), *SpirvShaderCompiler::kBindlessCBPrefix, CBIndex, *CBName);
-		const FString BindlessCBType = PrefixedCBName + TEXT("_Type");
-		const FString BindlessCBHeapName = PrefixedCBName + SpirvShaderCompiler::kBindlessHeapSuffix;
-		const FString PaddingName = FString::Printf(TEXT("%s_Padding"), *CBName);
-
-		FString CBDecl;
-		CBDecl.Reserve(Members.Len() * 3);  // start somewhere approx less bad
-
-		// Declare the struct
-		CBDecl += TEXT("struct ") + BindlessCBType + TEXT(" \n{\n") + Members + TEXT("\n};\n");
-
-		// Declare the safetype and bindless array for this cb
-		CBDecl += FString::Printf(TEXT("ConstantBuffer<%s> %s[];\n"), *BindlessCBType, *BindlessCBHeapName);
-
-		// Now bring in the CB
-		CBDecl += FString::Printf(TEXT("static const %s %s = %s[VulkanHitGroupSystemParameters.BindlessUniformBuffers[%d]];\n"),
-			*BindlessCBType, *PrefixedCBName, *BindlessCBHeapName, CBIndex);
-
-		// Now create a global scope var for each value (as the cbuffer would provide) to patch in seemlessly with the rest of the code
-		uint32 MemberOffset = 0;
-		const TCHAR* MemberSearchPtr = *Members;
-		const uint32 LastMemberSemicolonIndex = Members.Find(TEXT(";"), ESearchCase::CaseSensitive, ESearchDir::FromEnd, -1);
-		check(LastMemberSemicolonIndex != INDEX_NONE);
-		const TCHAR* LastMemberSemicolon = &Members[LastMemberSemicolonIndex];
-
-		do
-		{
-			const TCHAR* MemberTypeStartPtr = nullptr;
-			const TCHAR* MemberTypeEndPtr = nullptr;
-			ParseHLSLTypeName(MemberSearchPtr, MemberTypeStartPtr, MemberTypeEndPtr);
-			const FString MemberTypeName = FString::ConstructFromPtrSize(MemberTypeStartPtr, MemberTypeEndPtr - MemberTypeStartPtr);
-
-			FString MemberName;
-			MemberSearchPtr = ParseHLSLSymbolName(MemberTypeEndPtr, MemberName);
-			check(MemberName.Len() > 0);
-
-			if (MemberName.StartsWith(PaddingName))
-			{
-				while (*MemberSearchPtr && *MemberSearchPtr != ';')
-				{
-					MemberSearchPtr++;
-				}
-			}
-			else
-			{
-				// Skip over trailing tokens and pick up arrays
-				FString ArrayDecl;
-				while (*MemberSearchPtr && *MemberSearchPtr != ';')
-				{
-					if (*MemberSearchPtr == '[')
-					{
-						ArrayDecl.AppendChar(*MemberSearchPtr);
-
-						MemberSearchPtr++;
-						while (*MemberSearchPtr && *MemberSearchPtr != ']')
-						{
-							ArrayDecl.AppendChar(*MemberSearchPtr);
-							MemberSearchPtr++;
-						}
-
-						ArrayDecl.AppendChar(*MemberSearchPtr);
-					}
-
-					MemberSearchPtr++;
-				}
-
-				CBDecl += FString::Printf(TEXT("static const %s %s%s = %s.%s;\n"), *MemberTypeName, *MemberName, *ArrayDecl, *PrefixedCBName, *MemberName);
-			}
-
-			MemberSearchPtr++;
-
-		} while (MemberSearchPtr < LastMemberSemicolon);
-
-		return CBDecl;
-	};
-
-	// replace "cbuffer" decl with a struct filled from bindless constant buffer
-	TArray<FString> BindlessUBs;
-	{
-		const FString UniformBufferDeclIdentifier = TEXT("cbuffer");
-
-		int32 SearchIndex = PreprocessedShaderSource.Find(UniformBufferDeclIdentifier, ESearchCase::CaseSensitive, ESearchDir::FromStart, -1);
-		while (SearchIndex != INDEX_NONE)
-		{
-			FString StructName;
-			const TCHAR* StructNameEndPtr = ParseHLSLSymbolName(&PreprocessedShaderSource[SearchIndex + UniformBufferDeclIdentifier.Len()], StructName);
-			check(StructName.Len() > 0);
-
-			const int32 CBIndex = BindlessUBs.Add(StructName);
-			check(CBIndex < 16);
-
-			const TCHAR* OpeningBracePtr = FCString::Strstr(&PreprocessedShaderSource[SearchIndex + UniformBufferDeclIdentifier.Len()], TEXT("{"));
-			check(OpeningBracePtr);
-			const TCHAR* ClosingBracePtr = FindMatchingClosingBrace(OpeningBracePtr + 1);
-			check(ClosingBracePtr);
-			const int32 ClosingBraceIndex = ClosingBracePtr - (*PreprocessedShaderSource);
-
-			const FString Members = FString::ConstructFromPtrSize(OpeningBracePtr + 1, ClosingBracePtr - OpeningBracePtr - 1);
-			const FString NewDecl = GenerateNewDecl(CBIndex, Members, StructName);
-
-			const int32 OldDeclLen = ClosingBraceIndex - SearchIndex + 1;
-			PreprocessedShaderSource.RemoveAt(SearchIndex, OldDeclLen, EAllowShrinking::No);
-			PreprocessedShaderSource.InsertAt(SearchIndex, NewDecl);
-
-			SearchIndex = PreprocessedShaderSource.Find(UniformBufferDeclIdentifier, ESearchCase::CaseSensitive, ESearchDir::FromStart, SearchIndex + NewDecl.Len());
-		}
-	}
-	return BindlessUBs;
-}
-
-
 // Helper function to know how much space to set aside in the shader record for a global
 static uint32 GetSizeForType(FStringView TypeName, FStringView ArraySize)
 {
@@ -328,133 +210,6 @@ static uint32 ConvertGlobalsToShaderRecord(const FShaderParameterParser& ShaderP
 }
 
 
-static void UpdateBindlessUBs(const FSpirvShaderCompilerInternalState& InternalState, SpirvShaderCompilerSerializedOutput& SerializedOutput, FShaderCompilerOutput& Output)
-{
-	checkf(SerializedOutput.Header.Bindings.Num() == 0, TEXT("Shaders using bindless UBs should have no other bindings."));
-	for (int32 CBIndex = 0; CBIndex < InternalState.AllBindlessUBs.Num(); CBIndex++)
-	{
-		const FString& CBName = InternalState.AllBindlessUBs[CBIndex];
-
-		// It's possible SPIRV compilation has optimized out a buffer from every shader in the group
-		if (SerializedOutput.UsedBindlessUB.Contains(CBName))
-		{
-			FVulkanShaderHeader::FUniformBufferInfo& Info = SerializedOutput.Header.UniformBufferInfos.AddZeroed_GetRef();
-			Info.LayoutHash = SpirvShaderCompiler::GetUBLayoutHash(InternalState.Input, CBName);
-			Info.BindlessCBIndex = CBIndex;
-
-			const int32 UBIndex = SerializedOutput.Header.UniformBufferInfos.Num() - 1;
-			Output.ParameterMap.AddParameterAllocation(CBName, UBIndex, 0, 1, EShaderParameterType::UniformBuffer);
-		}
-	}
-}
-
-
-static bool CompileShaderGroup(
-	FSpirvShaderCompilerInternalState& InternalState,
-	const FShaderSource::FStringType& OriginalPreprocessedShaderSource,
-	FShaderCompilerOutput& MergedOutput
-)
-{
-	checkf(InternalState.bSupportsBindless && InternalState.bUseBindlessUniformBuffer, TEXT("Ray tracing requires full bindless in Vulkan."));
-
-	// Compile each one of the shader modules seperately and create one big blob for the engine
-	auto CompilePartialExport = [&OriginalPreprocessedShaderSource, &InternalState, &MergedOutput](
-		FSpirvShaderCompilerInternalState::EHitGroupShaderType HitGroupShaderType,
-		const TCHAR* PartialFileExtension,
-		SpirvShaderCompilerSerializedOutput& PartialSerializedOutput)
-	{
-		InternalState.HitGroupShaderType = HitGroupShaderType;
-
-		FShaderCompilerOutput TempOutput;
-		const bool bIsClosestHit = (HitGroupShaderType == FSpirvShaderCompilerInternalState::EHitGroupShaderType::ClosestHit);
-		FShaderCompilerOutput& PartialOutput = bIsClosestHit ? MergedOutput : TempOutput;
-
-		FShaderSource::FViewType OrigSourceView(OriginalPreprocessedShaderSource);
-		FShaderSource PartialPreprocessedShaderSource(OrigSourceView);
-		UE::ShaderCompilerCommon::RemoveDeadCode(PartialPreprocessedShaderSource, InternalState.GetEntryPointName(), PartialOutput.Errors);
-
-		if (InternalState.bDebugDump)
-		{
-			DumpDebugShaderText(InternalState.Input, PartialPreprocessedShaderSource.GetView().GetData(), *FString::Printf(TEXT("%s.hlsl"), PartialFileExtension));
-		}
-
-		const bool bPartialSuccess = SpirvShaderCompiler::CompileWithShaderConductor(InternalState, PartialPreprocessedShaderSource.GetView(), PartialSerializedOutput, PartialOutput);
-
-		if (!bIsClosestHit)
-		{
-			MergedOutput.NumInstructions = FMath::Max(MergedOutput.NumInstructions, PartialOutput.NumInstructions);
-			MergedOutput.NumTextureSamplers = FMath::Max(MergedOutput.NumTextureSamplers, PartialOutput.NumTextureSamplers);
-			MergedOutput.Errors.Append(MoveTemp(PartialOutput.Errors));
-		}
-
-		return bPartialSuccess;
-	};
-
-	bool bSuccess = false;
-
-	// Closest Hit Module, always present
-	SpirvShaderCompilerSerializedOutput ClosestHitSerializedOutput;
-	{
-		bSuccess = CompilePartialExport(FSpirvShaderCompilerInternalState::EHitGroupShaderType::ClosestHit, TEXT("closest"), ClosestHitSerializedOutput);
-	}
-
-	// Any Hit Module, optional
-	const bool bHasAnyHitModule = !InternalState.AnyHitEntry.IsEmpty();
-	SpirvShaderCompilerSerializedOutput AnyHitSerializedOutput;
-	if (bSuccess && bHasAnyHitModule)
-	{
-		bSuccess = CompilePartialExport(FSpirvShaderCompilerInternalState::EHitGroupShaderType::AnyHit, TEXT("anyhit"), AnyHitSerializedOutput);
-	}
-
-	// Intersection Module, optional
-	const bool bHasIntersectionModule = !InternalState.IntersectionEntry.IsEmpty();
-	SpirvShaderCompilerSerializedOutput IntersectionSerializedOutput;
-	if (bSuccess && bHasIntersectionModule)
-	{
-		bSuccess = CompilePartialExport(FSpirvShaderCompilerInternalState::EHitGroupShaderType::Intersection, TEXT("intersection"), IntersectionSerializedOutput);
-	}
-
-	// Collapse the bindless UB usage into one set and then update the headers
-	ClosestHitSerializedOutput.UsedBindlessUB.Append(AnyHitSerializedOutput.UsedBindlessUB);
-	ClosestHitSerializedOutput.UsedBindlessUB.Append(IntersectionSerializedOutput.UsedBindlessUB);
-	UpdateBindlessUBs(InternalState, ClosestHitSerializedOutput, MergedOutput);
-
-	{
-		// :todo-jn: Having multiple entrypoints in a single SPIRV blob crashes on FLumenHardwareRayTracingMaterialHitGroup for some reason
-		// Adjust the header before we write it out
-		ClosestHitSerializedOutput.Header.RayGroupAnyHit = bHasAnyHitModule ? FVulkanShaderHeader::ERayHitGroupEntrypoint::SeparateBlob : FVulkanShaderHeader::ERayHitGroupEntrypoint::NotPresent;
-		ClosestHitSerializedOutput.Header.RayGroupIntersection = bHasIntersectionModule ? FVulkanShaderHeader::ERayHitGroupEntrypoint::SeparateBlob : FVulkanShaderHeader::ERayHitGroupEntrypoint::NotPresent;
-
-		check(ClosestHitSerializedOutput.Spirv.Data.Num() != 0);
-		FMemoryWriter Ar(MergedOutput.ShaderCode.GetWriteAccess(), true);
-		Ar << ClosestHitSerializedOutput.Header;
-		Ar << ClosestHitSerializedOutput.ShaderResourceTable;
-
-		{
-			uint32 SpirvCodeSizeBytes = ClosestHitSerializedOutput.Spirv.GetByteSize();
-			Ar << SpirvCodeSizeBytes;
-			Ar.Serialize((uint8*)ClosestHitSerializedOutput.Spirv.Data.GetData(), SpirvCodeSizeBytes);
-		}
-
-		if (bHasAnyHitModule)
-		{
-			uint32 SpirvCodeSizeBytes = AnyHitSerializedOutput.Spirv.GetByteSize();
-			Ar << SpirvCodeSizeBytes;
-			Ar.Serialize((uint8*)AnyHitSerializedOutput.Spirv.Data.GetData(), SpirvCodeSizeBytes);
-		}
-
-		if (bHasIntersectionModule)
-		{
-			uint32 SpirvCodeSizeBytes = IntersectionSerializedOutput.Spirv.GetByteSize();
-			Ar << SpirvCodeSizeBytes;
-			Ar.Serialize((uint8*)IntersectionSerializedOutput.Spirv.Data.GetData(), SpirvCodeSizeBytes);
-		}
-	}
-
-	MergedOutput.bSucceeded = bSuccess;
-	return bSuccess;
-}
-
 struct FVulkanShaderParameterParserPlatformConfiguration : public FShaderParameterParser::FPlatformConfiguration
 {
 	FVulkanShaderParameterParserPlatformConfiguration(const FShaderCompilerInput& Input, TMap<FStringView, FStringView>& InReplacedGlobals)
@@ -539,7 +294,7 @@ void CompileVulkanShader(const FShaderCompilerInput& Input, const FShaderPreproc
 	if (InternalState.bUseBindlessUniformBuffer)
 	{
 		InternalState.ShaderRecordGlobalsSize = ConvertGlobalsToShaderRecord(ShaderParameterParser, ReplacedGlobals, PreprocessedSource, Output);
-		InternalState.AllBindlessUBs = ConvertUBToBindless(PreprocessedSource);
+		InternalState.AllBindlessUBs = SpirvShaderCompiler::ConvertUBToBindless(PreprocessedSource);
 	}
 
 	if (ShaderParameterParser.DidModifyShader() || InternalState.AllBindlessUBs.Num() || InternalState.ShaderRecordGlobalsSize)
@@ -562,7 +317,7 @@ void CompileVulkanShader(const FShaderCompilerInput& Input, const FShaderPreproc
 	// HitGroup shaders might have multiple entrypoints that we combine into a single blob
 	if (InternalState.HasMultipleEntryPoints())
 	{
-		bSuccess = CompileShaderGroup(InternalState, PreprocessedSourceToCompile, Output);
+		bSuccess = SpirvShaderCompiler::CompileShaderGroup(InternalState, PreprocessedSourceToCompile, Output);
 	}
 	else
 	{
@@ -572,7 +327,7 @@ void CompileVulkanShader(const FShaderCompilerInput& Input, const FShaderPreproc
 
 		if (InternalState.bUseBindlessUniformBuffer)
 		{
-			UpdateBindlessUBs(InternalState, SerializedOutput, Output);
+			SpirvShaderCompiler::UpdateBindlessUBs(InternalState, SerializedOutput, Output);
 		}
 
 		// Write out the header and shader source code (except for the extra shaders in hit groups)
