- Moved some includes to the .inl where they are needed

diff --git a/Engine/Source/Developer/VulkanShaderFormat/Private/SpirVShaderCompiler.inl b/Engine/Source/Developer/VulkanShaderFormat/Private/SpirVShaderCompiler.inl
--- a/Engine/Source/Developer/VulkanShaderFormat/Private/SpirVShaderCompiler.inl
+++ b/Engine/Source/Developer/VulkanShaderFormat/Private/SpirVShaderCompiler.inl
@@ -2,6 +2,9 @@
 // .
 
 #include "SpirvCommon.h"
+#include "Serialization/MemoryWriter.h"
+#include "Misc/FileHelper.h"
+#include "Misc/Paths.h"
 
 // A collection of states and data that is locked in at the top level call and doesn't change throughout the compilation process
 class FSpirvShaderCompilerInternalState
diff --git a/Engine/Source/Developer/VulkanShaderFormat/Private/VulkanShaderCompiler.cpp b/Engine/Source/Developer/VulkanShaderFormat/Private/VulkanShaderCompiler.cpp
--- a/Engine/Source/Developer/VulkanShaderFormat/Private/VulkanShaderCompiler.cpp
+++ b/Engine/Source/Developer/VulkanShaderFormat/Private/VulkanShaderCompiler.cpp
@@ -12,9 +12,6 @@
 #include "VulkanThirdParty.h"
 #include "VulkanBackend.h"
 #include "VulkanShaderResources.h"
-#include "Serialization/MemoryWriter.h"
-#include "Misc/FileHelper.h"
-#include "Misc/Paths.h"
 
 #include "SpirVShaderCompiler.inl"
 
