VulkanRHI: Add "-fspv-use-legacy-buffer-matrix-order" to shader compilation flags to fix column vs row major reads from byteaddressbuffers.  Fix Nanite skinning in Vulkan.

#rnx

diff --git a/Engine/Shaders/Public/Platform/Vulkan/VulkanCommon.ush b/Engine/Shaders/Public/Platform/Vulkan/VulkanCommon.ush
--- a/Engine/Shaders/Public/Platform/Vulkan/VulkanCommon.ush
+++ b/Engine/Shaders/Public/Platform/Vulkan/VulkanCommon.ush
@@ -7,7 +7,7 @@
 #pragma once
 
 // Update this GUID to invalidate shader recompilation for only Vulkan shaders
-#pragma message("UESHADERMETADATA_VERSION 88236FE1-81A3-418B-896E-F41A3AD65ED6")
+#pragma message("UESHADERMETADATA_VERSION 270F9A3B-284D-41E3-91B3-3ADFF4AD8784")
 
 #define COMPILER_SUPPORTS_ULONG_TYPES 1
 #define COMPILER_SUPPORTS_UINT64_IMAGE_ATOMICS (1 && COMPILER_SUPPORTS_ULONG_TYPES)
diff --git a/Engine/Source/Developer/VulkanShaderFormat/Private/SpirVShaderCompiler.inl b/Engine/Source/Developer/VulkanShaderFormat/Private/SpirVShaderCompiler.inl
--- a/Engine/Source/Developer/VulkanShaderFormat/Private/SpirVShaderCompiler.inl
+++ b/Engine/Source/Developer/VulkanShaderFormat/Private/SpirVShaderCompiler.inl
@@ -1225,6 +1225,10 @@ static bool CompileWithShaderConductor(
 		ExtraDxcArgs.Add(TEXT("-fvk-allow-rwstructuredbuffer-arrays"));
 	}
 
+	// Fix issues when reading matrices directly for ByteAddrBuffer
+	// By default the compiler will emit column-major loads and this flag makes sure to revert to the original behavior of row-major.
+	ExtraDxcArgs.Add(TEXT("-fspv-use-legacy-buffer-matrix-order"));
+
 	// Load shader source into compiler context
 	CompilerContext.LoadSource(PreprocessedShader, Input.VirtualSourceFilePath, InternalState.GetEntryPointName(), InternalState.GetShaderFrequency(), &AdditionalDefines, &ExtraDxcArgs);
 
