VulkanRHI: fill EShaderResourceUsageFlags at shader compile time to accurately set FRHIShader flags at runtime (HasNoDerivates,ShaderBundles).

#rnx

diff --git a/Engine/Shaders/Public/Platform/Vulkan/VulkanCommon.ush b/Engine/Shaders/Public/Platform/Vulkan/VulkanCommon.ush
--- a/Engine/Shaders/Public/Platform/Vulkan/VulkanCommon.ush
+++ b/Engine/Shaders/Public/Platform/Vulkan/VulkanCommon.ush
@@ -7,7 +7,7 @@
 #pragma once
 
 // Update this GUID to invalidate shader recompilation for only Vulkan shaders
-#pragma message("UESHADERMETADATA_VERSION 270F9A3B-284D-41E3-91B3-3ADFF4AD8784")
+#pragma message("UESHADERMETADATA_VERSION E4EECD20-AD45-4120-9A4E-FBCA3FFB670D")
 
 #define COMPILER_SUPPORTS_ULONG_TYPES 1
 #define COMPILER_SUPPORTS_UINT64_IMAGE_ATOMICS (1 && COMPILER_SUPPORTS_ULONG_TYPES)
diff --git a/Engine/Source/Developer/VulkanShaderFormat/Private/SpirVShaderCompiler.inl b/Engine/Source/Developer/VulkanShaderFormat/Private/SpirVShaderCompiler.inl
--- a/Engine/Source/Developer/VulkanShaderFormat/Private/SpirVShaderCompiler.inl
+++ b/Engine/Source/Developer/VulkanShaderFormat/Private/SpirVShaderCompiler.inl
@@ -137,6 +137,7 @@ struct SpirvShaderCompilerSerializedOutput
 	SpirvShaderCompilerSerializedOutput()
 		: Header(FVulkanShaderHeader::EZero)
 	{
+		FMemory::Memzero(PackedResourceCounts);
 	}
 
 	FVulkanShaderHeader Header;  // TODO: Convert descriptors into more generic Spirv information
@@ -145,6 +146,7 @@ struct SpirvShaderCompilerSerializedOutput
 	FSpirv Spirv;
 	uint32 SpirvCRC = 0;
 	const ANSICHAR* SpirvEntryPointName = nullptr;
+	FShaderCodePackedResourceCounts PackedResourceCounts;
 
 	TSet<FString> UsedBindlessUB;
 };
@@ -186,6 +188,62 @@ static uint32 GetUBLayoutHash(const FShaderCompilerInput& ShaderInput, const FSt
 	return LayoutHash;
 }
 
+static bool HasDerivatives(const FSpirv& Spirv)
+{
+	for (FSpirvConstIterator Iter = Spirv.begin(); Iter != Spirv.end(); ++Iter)
+	{
+		switch (Iter.Opcode())
+		{
+		case SpvOpCapability:
+		{
+			const uint32 Capability = Iter.Operand(1);
+			if ((Capability == SpvCapabilityComputeDerivativeGroupLinearNV) ||
+				(Capability == SpvCapabilityComputeDerivativeGroupQuadsNV))
+			{
+				return true;
+			}
+		}
+		break;
+
+		case SpvOpExtension:
+		case SpvOpEntryPoint:
+			// By the time we've reached extensions/entrypoints, we're done listing capabilities
+			return false;
+
+		default:
+			break;
+		}
+	}
+	return false;
+}
+
+static void FillShaderResourceUsageFlags(const FSpirvShaderCompilerInternalState& InternalState, SpirvShaderCompilerSerializedOutput& SerializedOutput)
+{
+	FShaderCodePackedResourceCounts& PackedResourceCounts = SerializedOutput.PackedResourceCounts;
+
+	if (InternalState.Input.Target.GetFrequency() == SF_Compute && 
+		InternalState.Input.Environment.CompilerFlags.Contains(CFLAG_CheckForDerivativeOps))
+	{
+		if (!HasDerivatives(SerializedOutput.Spirv))
+		{
+			PackedResourceCounts.UsageFlags |= EShaderResourceUsageFlags::NoDerivativeOps;
+		}
+	}
+
+	if (InternalState.bSupportsBindless)
+	{
+		PackedResourceCounts.UsageFlags |= EShaderResourceUsageFlags::BindlessResources;
+		PackedResourceCounts.UsageFlags |= EShaderResourceUsageFlags::BindlessSamplers;
+	}
+
+	if (InternalState.Input.Environment.CompilerFlags.Contains(CFLAG_ShaderBundle))
+	{
+		PackedResourceCounts.UsageFlags |= EShaderResourceUsageFlags::ShaderBundle;
+	}
+
+	// TODO: When DiagnosticBuffer is supported:
+	// PackedResourceCounts.UsageFlags |= EShaderResourceUsageFlags::DiagnosticBuffer;
+}
 
 static void BuildShaderOutput(
 	SpirvShaderCompilerSerializedOutput& SerializedOutput,
@@ -795,6 +853,7 @@ static bool BuildShaderOutputFromSpirv(
 		const int32 GlobalUBCount = AddReflectionInfos(Bindings.UniformBuffers, VK_DESCRIPTOR_TYPE_UNIFORM_BUFFER, 0, true);
 		const int32 UBOBindings = AddReflectionInfos(Bindings.UniformBuffers, VK_DESCRIPTOR_TYPE_UNIFORM_BUFFER, GlobalUBCount);
 		SerializedOutput.Header.NumBoundUniformBuffers = UBOBindings;
+		SerializedOutput.PackedResourceCounts.NumCBs = (uint8)UBOBindings;
 
 		AddReflectionInfos(Bindings.InputAttachments, VK_DESCRIPTOR_TYPE_INPUT_ATTACHMENT, 0);
 
@@ -802,13 +861,17 @@ static bool BuildShaderOutputFromSpirv(
 		UAVBindings = AddReflectionInfos(Bindings.TBufferUAVs, VK_DESCRIPTOR_TYPE_STORAGE_TEXEL_BUFFER, UAVBindings);
 		UAVBindings = AddReflectionInfos(Bindings.SBufferUAVs, VK_DESCRIPTOR_TYPE_STORAGE_BUFFER, UAVBindings);
 		UAVBindings = AddReflectionInfos(Bindings.TextureUAVs, VK_DESCRIPTOR_TYPE_STORAGE_IMAGE, UAVBindings);
+		SerializedOutput.PackedResourceCounts.NumUAVs = (uint8)UAVBindings;
 
 		int32 SRVBindings = 0;
 		SRVBindings = AddReflectionInfos(Bindings.TBufferSRVs, VK_DESCRIPTOR_TYPE_UNIFORM_TEXEL_BUFFER, SRVBindings);
 		checkf(Bindings.SBufferSRVs.IsEmpty(), TEXT("GatherSpirvReflectionBindings should have dumped all SBufferSRVs into SBufferUAVs."));
 		SRVBindings = AddReflectionInfos(Bindings.TextureSRVs, VK_DESCRIPTOR_TYPE_SAMPLED_IMAGE, SRVBindings);
+		SerializedOutput.PackedResourceCounts.NumSRVs = (uint8)SRVBindings;
 
 		Output.NumTextureSamplers = AddReflectionInfos(Bindings.Samplers, VK_DESCRIPTOR_TYPE_SAMPLER, 0);
+		SerializedOutput.PackedResourceCounts.NumSamplers = (uint8)Output.NumTextureSamplers;
+
 		AddReflectionInfos(Bindings.AccelerationStructures, VK_DESCRIPTOR_TYPE_ACCELERATION_STRUCTURE_KHR, 0);
 	}
 
diff --git a/Engine/Source/Runtime/VulkanRHI/Private/VulkanShaders.cpp b/Engine/Source/Runtime/VulkanRHI/Private/VulkanShaders.cpp
--- a/Engine/Source/Runtime/VulkanRHI/Private/VulkanShaders.cpp
+++ b/Engine/Source/Runtime/VulkanRHI/Private/VulkanShaders.cpp
@@ -40,6 +40,25 @@ FVulkanShaderFactory::~FVulkanShaderFactory()
 	}
 }
 
+
+static void ReadShaderOptionalData(FShaderCodeReader& ShaderCode, FRHIShader* RHIShader)
+{
+	const FShaderCodePackedResourceCounts* PackedResourceCounts = ShaderCode.FindOptionalData<FShaderCodePackedResourceCounts>();
+	if (PackedResourceCounts)
+	{
+		if (RHIShader->GetFrequency() == SF_Compute)
+		{
+			RHIShader->SetNoDerivativeOps(EnumHasAnyFlags(PackedResourceCounts->UsageFlags, EShaderResourceUsageFlags::NoDerivativeOps));
+		}
+		RHIShader->SetShaderBundleUsage(EnumHasAnyFlags(PackedResourceCounts->UsageFlags, EShaderResourceUsageFlags::ShaderBundle));
+	}
+
+#if RHI_INCLUDE_SHADER_DEBUG_DATA
+	RHIShader->Debug.ShaderName = ShaderCode.FindOptionalData(FShaderCodeName::Key);
+	UE::RHICore::SetupShaderCodeValidationData(RHIShader, ShaderCode);
+#endif
+}
+
 template <typename ShaderType> 
 ShaderType* FVulkanShaderFactory::CreateShader(TArrayView<const uint8> Code, FVulkanDevice* Device)
 {
@@ -75,11 +94,8 @@ ShaderType* FVulkanShaderFactory::CreateShader(TArrayView<const uint8> Code, FVu
 
 				ShaderMap[ShaderType::StaticFrequency].Add(ShaderKey, RetShader);
 
-#if !(UE_BUILD_SHIPPING || UE_BUILD_TEST)
 				FShaderCodeReader ShaderCode(Code);
-				RetShader->Debug.ShaderName = ShaderCode.FindOptionalData(FShaderCodeName::Key);
-				UE::RHICore::SetupShaderCodeValidationData(RetShader, ShaderCode);
-#endif
+				ReadShaderOptionalData(ShaderCode, RetShader);
 			}
 		}
 	}
@@ -160,11 +176,8 @@ FVulkanRayTracingShader* FVulkanShaderFactory::CreateRayTracingShader(TArrayView
 
 				ShaderMap[ShaderFrequency].Add(ShaderKey, RetShader);
 
-#if !(UE_BUILD_SHIPPING || UE_BUILD_TEST)
 				FShaderCodeReader ShaderCode(Code);
-				RetShader->Debug.ShaderName = ShaderCode.FindOptionalData(FShaderCodeName::Key);
-                UE::RHICore::SetupShaderCodeValidationData(RetShader, ShaderCode);
-#endif
+				ReadShaderOptionalData(ShaderCode, RetShader);
 			}
 		}
 	}
