VulkanRHI: Add missing file from 35154013 (fill EShaderResourceUsageFlags at shader compile time)

#rnx

diff --git a/Engine/Source/Developer/VulkanShaderFormat/Private/VulkanShaderCompiler.cpp b/Engine/Source/Developer/VulkanShaderFormat/Private/VulkanShaderCompiler.cpp
--- a/Engine/Source/Developer/VulkanShaderFormat/Private/VulkanShaderCompiler.cpp
+++ b/Engine/Source/Developer/VulkanShaderFormat/Private/VulkanShaderCompiler.cpp
@@ -339,6 +339,9 @@ void CompileVulkanShader(const FShaderCompilerInput& Input, const FShaderPreproc
 		{
 			Ar.Serialize((uint8*)SerializedOutput.Spirv.Data.GetData(), SpirvCodeSizeBytes);
 		}
+
+		SpirvShaderCompiler::FillShaderResourceUsageFlags(InternalState, SerializedOutput);
+		Output.ShaderCode.AddOptionalData(SerializedOutput.PackedResourceCounts);
 	}
 #endif // PLATFORM_MAC || PLATFORM_WINDOWS || PLATFORM_LINUX
 	
