VulkanRHI: Add missing 3 component types.

#rnx

diff --git a/Engine/Source/Developer/VulkanShaderFormat/Private/VulkanShaderCompiler.cpp b/Engine/Source/Developer/VulkanShaderFormat/Private/VulkanShaderCompiler.cpp
--- a/Engine/Source/Developer/VulkanShaderFormat/Private/VulkanShaderCompiler.cpp
+++ b/Engine/Source/Developer/VulkanShaderFormat/Private/VulkanShaderCompiler.cpp
@@ -157,9 +157,11 @@ static uint32 GetSizeForType(FStringView TypeName, FStringView ArraySize)
 	{
 		SizeForTypeMap.Add(FStringView(TEXT("uint")),   4);
 		SizeForTypeMap.Add(FStringView(TEXT("uint2")),  8);
+		SizeForTypeMap.Add(FStringView(TEXT("uint3")),  12);
 		SizeForTypeMap.Add(FStringView(TEXT("uint4")),  16);
 		SizeForTypeMap.Add(FStringView(TEXT("float")),  4);
 		SizeForTypeMap.Add(FStringView(TEXT("float2")), 8);
+		SizeForTypeMap.Add(FStringView(TEXT("float3")), 12);
 		SizeForTypeMap.Add(FStringView(TEXT("float4")), 16);
 	}
 
