[ShaderAnalyzer] Allow returning intermediate outputs in FShaderCompilerOutput instead of just dumping them to files.

This serves as preparation for the ShaderAnalyzer to return intermediate outputs (SPIR-V, GLSL source etc.) in FShaderCompilerOutput, so it can be transferred between the SCW and editor efficiently.
It makes use of the existing ShaderStatistics and only adds a new compiler flag to enable full shader reflection output. No SCW input/output version update is required.

#rnx

diff --git a/Engine/Shaders/Public/ShaderVersion.ush b/Engine/Shaders/Public/ShaderVersion.ush
--- a/Engine/Shaders/Public/ShaderVersion.ush
+++ b/Engine/Shaders/Public/ShaderVersion.ush
@@ -3,4 +3,4 @@
 // Each time the console command is executed it generates a new GUID. As this file is included
 // in Platform.ush (which should be included in any shader) it allows to invalidate the shader DDC.
 // 
-#pragma message("UESHADERMETADATA_VERSION E9D355B4-B3C1-4E3B-AEA9-147D42EA5F07")
+#pragma message("UESHADERMETADATA_VERSION 4EA1C1A3-BE93-4931-8B40-7CFE5B6C2E85")
diff --git a/Engine/Source/Developer/Apple/MetalShaderFormat/Private/MetalCompileShaderSPIRV.cpp b/Engine/Source/Developer/Apple/MetalShaderFormat/Private/MetalCompileShaderSPIRV.cpp
--- a/Engine/Source/Developer/Apple/MetalShaderFormat/Private/MetalCompileShaderSPIRV.cpp
+++ b/Engine/Source/Developer/Apple/MetalShaderFormat/Private/MetalCompileShaderSPIRV.cpp
@@ -279,16 +279,22 @@ void FMetalCompileShaderSPIRV::DoCompileMetalShader(
 		}
 		
 		// If source data was modified, reload it into the compiler context
+		const bool bOutputAnalysisArtifacts = Input.Environment.CompilerFlags.Contains(CFLAG_OutputAnalysisArtifacts);
 		if (bSourceDataWasModified)
 		{
 			CompilerContext.LoadSource(FAnsiStringView(*SourceData, SourceData.Len()), Input.VirtualSourceFilePath, Input.EntryPointName, Frequency);
+
+			if (bOutputAnalysisArtifacts)
+			{
+				Output.AddStatistic<FString>(TEXT("Modified HLSL"), ANSI_TO_TCHAR(&SourceData[0]), FGenericShaderStat::EFlags::Hidden, FShaderStatTagNames::AnalysisArtifactsName);
+			}
 		}
 
 		if (bDumpDebugInfo)
 		{
 			DumpDebugShaderText(Input, &SourceData[0], SourceData.Len(), TEXT("rewritten.hlsl"));
 		}
-		
+
 		CrossCompiler::FHlslccHeaderWriter CCHeaderWriter;
 
 		FString ALNString;
@@ -309,6 +315,15 @@ void FMetalCompileShaderSPIRV::DoCompileMetalShader(
 				DumpDebugShaderDisassembledSpirv(Input, SpirvData.GetData(), SpirvData.Num() * sizeof(uint32), TEXT("spvasm"));
 			}
 			
+			if (bOutputAnalysisArtifacts)
+			{
+				FGenericShaderStat ShaderIRReflection;
+				if (CompilerContext.Disassemble(CrossCompiler::EShaderConductorIR::Spirv, SpirvData.GetData(), SpirvData.Num() * sizeof(uint32), ShaderIRReflection))
+				{
+					Output.ShaderStatistics.Add(MoveTemp(ShaderIRReflection));
+				}
+			}
+
 			// Now perform reflection on the SPIRV and tweak any decorations that we need to.
 			// This used to be done via JSON, but that was slow and alloc happy so use SPIRV-Reflect instead.
 			spv_reflect::ShaderModule Reflection(SpirvData.Num() * sizeof(uint32), SpirvData.GetData());
diff --git a/Engine/Source/Developer/ShaderCompilerCommon/Private/ShaderConductorContext.cpp b/Engine/Source/Developer/ShaderCompilerCommon/Private/ShaderConductorContext.cpp
--- a/Engine/Source/Developer/ShaderCompilerCommon/Private/ShaderConductorContext.cpp
+++ b/Engine/Source/Developer/ShaderCompilerCommon/Private/ShaderConductorContext.cpp
@@ -1123,6 +1123,30 @@ namespace CrossCompiler
 		}
 	}
 
+	static const TCHAR* LexToString(EShaderConductorIR Language)
+	{
+		switch (Language)
+		{
+		case EShaderConductorIR::Spirv: return TEXT("SPIR-V");
+		case EShaderConductorIR::Dxil: return TEXT("DXIL");
+		default: return TEXT("");
+		}
+	}
+
+	bool FShaderConductorContext::Disassemble(EShaderConductorIR Language, const void* Binary, uint32 BinaryByteSize, FGenericShaderStat& OutShaderStat)
+	{
+		TArray<ANSICHAR> AssemblyText;
+		if (FShaderConductorContext::Disassemble(Language, Binary, BinaryByteSize, AssemblyText))
+		{
+			OutShaderStat.StatName = LexToString(Language);
+			OutShaderStat.Value.Set<FString>(ANSI_TO_TCHAR(AssemblyText.GetData()));
+			OutShaderStat.Flags = FGenericShaderStat::EFlags::Hidden;
+			OutShaderStat.TagName = FShaderStatTagNames::AnalysisArtifactsName;
+			return true;
+		}
+		return false;
+	}
+
 	void FShaderConductorContext::Shutdown()
 	{
 #if PLATFORM_LINUX
diff --git a/Engine/Source/Developer/ShaderCompilerCommon/Public/ShaderConductorContext.h b/Engine/Source/Developer/ShaderCompilerCommon/Public/ShaderConductorContext.h
--- a/Engine/Source/Developer/ShaderCompilerCommon/Public/ShaderConductorContext.h
+++ b/Engine/Source/Developer/ShaderCompilerCommon/Public/ShaderConductorContext.h
@@ -259,9 +259,12 @@ namespace CrossCompiler
 		/** Returns the table of special identifiers generated by ShaderConductor. */
 		static const FShaderConductorIdentifierTable& GetIdentifierTable();
 
-		/** Disassembles the specified SPIR-V module and returns its assembly as text representation. */
+		/** Disassembles the specified SPIR-V or DXIL module and returns its assembly as text representation. */
 		static bool Disassemble(EShaderConductorIR Language, const void* Binary, uint32 BinaryByteSize, TArray<ANSICHAR>& OutAssemblyText);
 
+		/** Helper function to disassemble the specified SPIR-V or DXIL module and add it to the output code reflection array. */
+		static bool Disassemble(EShaderConductorIR Language, const void* Binary, uint32 BinaryByteSize, FGenericShaderStat& OutShaderStat);
+
 		/** Returns a filename extension for the specified shading language and shader stage, e.g. "frag" for a GLSL pixel shader. */
 		static const TCHAR* GetShaderFileExt(EShaderConductorLanguage Language, EShaderFrequency ShaderStage = SF_NumFrequencies);
 
diff --git a/Engine/Source/Developer/VulkanShaderFormat/Private/SpirVShaderCompiler.inl b/Engine/Source/Developer/VulkanShaderFormat/Private/SpirVShaderCompiler.inl
--- a/Engine/Source/Developer/VulkanShaderFormat/Private/SpirVShaderCompiler.inl
+++ b/Engine/Source/Developer/VulkanShaderFormat/Private/SpirVShaderCompiler.inl
@@ -1343,6 +1343,19 @@ static bool CompileWithShaderConductor(
 
 	// Flush warnings
 	CompilerContext.FlushErrors(Output.Errors);
+
+	// Return code reflection if requested for shader analysis
+	if (Input.Environment.CompilerFlags.Contains(CFLAG_OutputAnalysisArtifacts) && Output.bSucceeded)
+	{
+		const TArray<uint32>& SpirvData = SerializedOutput.Spirv.Data;
+		FGenericShaderStat ShaderReflection;
+		if (CrossCompiler::FShaderConductorContext::Disassemble(CrossCompiler::EShaderConductorIR::Spirv, SpirvData.GetData(), SpirvData.Num() * SpirvData.GetTypeSize(), ShaderReflection))
+		{
+			ShaderReflection.StatName = FName(FString::Printf(TEXT("%s (%s)"), *ShaderReflection.StatName.ToString(), *InternalState.Input.EntryPointName));
+			Output.ShaderStatistics.Add(MoveTemp(ShaderReflection));
+		}
+	}
+
 	return true;
 }
 
@@ -1657,6 +1670,42 @@ static bool CompileShaderGroup(
 		}
 	}
 
+	// Return code reflection if requested for shader analysis
+	if (InternalState.Input.Environment.CompilerFlags.Contains(CFLAG_OutputAnalysisArtifacts) && bSuccess)
+	{
+		{
+			const TArray<uint32>& SpirvData = ClosestHitSerializedOutput.Spirv.Data;
+			FGenericShaderStat ClosestHitReflection;
+			if (CrossCompiler::FShaderConductorContext::Disassemble(CrossCompiler::EShaderConductorIR::Spirv, SpirvData.GetData(), SpirvData.Num() * SpirvData.GetTypeSize(), ClosestHitReflection))
+			{
+				ClosestHitReflection.StatName = FName(FString::Printf(TEXT("%s (%s)"), *ClosestHitReflection.StatName.ToString(), *InternalState.GetEntryPointName()));
+				MergedOutput.ShaderStatistics.Add(MoveTemp(ClosestHitReflection));
+			}
+		}
+
+		if (bHasAnyHitModule)
+		{
+			const TArray<uint32>& SpirvData = AnyHitSerializedOutput.Spirv.Data;
+			FGenericShaderStat AnyHitReflection;
+			if (CrossCompiler::FShaderConductorContext::Disassemble(CrossCompiler::EShaderConductorIR::Spirv, SpirvData.GetData(), SpirvData.Num() * SpirvData.GetTypeSize(), AnyHitReflection))
+			{
+				AnyHitReflection.StatName = FName(FString::Printf(TEXT("%s (%s)"), *AnyHitReflection.StatName.ToString(), *InternalState.AnyHitEntry));
+				MergedOutput.ShaderStatistics.Add(MoveTemp(AnyHitReflection));
+			}
+		}
+
+		if (bHasIntersectionModule)
+		{
+			const TArray<uint32>& SpirvData = IntersectionSerializedOutput.Spirv.Data;
+			FGenericShaderStat IntersectionReflection;
+			if (CrossCompiler::FShaderConductorContext::Disassemble(CrossCompiler::EShaderConductorIR::Spirv, SpirvData.GetData(), SpirvData.Num()* SpirvData.GetTypeSize(), IntersectionReflection))
+			{
+				IntersectionReflection.StatName = FName(FString::Printf(TEXT("%s (%s)"), *IntersectionReflection.StatName.ToString(), *InternalState.IntersectionEntry));
+				MergedOutput.ShaderStatistics.Add(MoveTemp(IntersectionReflection));
+			}
+		}
+	}
+
 	MergedOutput.bSucceeded = bSuccess;
 	return bSuccess;
 }
diff --git a/Engine/Source/Developer/Windows/ShaderFormatD3D/Private/D3DShaderCompiler.cpp b/Engine/Source/Developer/Windows/ShaderFormatD3D/Private/D3DShaderCompiler.cpp
--- a/Engine/Source/Developer/Windows/ShaderFormatD3D/Private/D3DShaderCompiler.cpp
+++ b/Engine/Source/Developer/Windows/ShaderFormatD3D/Private/D3DShaderCompiler.cpp
@@ -902,6 +902,12 @@ static bool CompileAndProcessD3DShaderFXCExt(
 				);
 			};
 
+			if (Input.Environment.CompilerFlags.Contains(CFLAG_OutputAnalysisArtifacts) && !bSecondPassAferUnusedInputRemoval)
+			{
+				// Add disassembly string to reflection output
+				Output.AddStatistic<FString>(TEXT("Optimized HLSL-2018"), ANSI_TO_TCHAR(CrossCompiledSource.GetData()), FGenericShaderStat::EFlags::Hidden, FShaderStatTagNames::AnalysisArtifactsName);
+			}
+
 			// Compile again with FXC - 1st try
 			const uint32 CompileFlagsNoWarningsAsErrors = CompileFlags & (~D3DCOMPILE_WARNINGS_ARE_ERRORS);
 			Result = CompileCrossCompiledHlsl(CrossCompiledSource, CompileFlagsNoWarningsAsErrors);
@@ -1157,6 +1163,23 @@ static bool CompileAndProcessD3DShaderFXCExt(
 		// Save results if compilation and reflection succeeded
 		if (Output.bSucceeded)
 		{
+			if (Input.Environment.CompilerFlags.Contains(CFLAG_OutputAnalysisArtifacts) && !bSecondPassAferUnusedInputRemoval)
+			{
+				TRefCountPtr<ID3DBlob> Disassembly;
+				if (SUCCEEDED(D3DDisassembleFunc(Shader->GetBufferPointer(), Shader->GetBufferSize(), 0, "", Disassembly.GetInitReference())))
+				{
+					// Convert disassembly from ID3DBlob to FString
+					TArray<ANSICHAR> DisassemblyAnsiString;
+					DisassemblyAnsiString.SetNumUninitialized(Disassembly->GetBufferSize() + 1);
+					FMemory::Memcpy(DisassemblyAnsiString.GetData(), Disassembly->GetBufferPointer(), Disassembly->GetBufferSize());
+					DisassemblyAnsiString[Disassembly->GetBufferSize()] = '\0';
+					FString DisassemblyString(DisassemblyAnsiString);
+
+					// Add disassembly string to reflection output
+					Output.AddStatistic<FString>(TEXT("DXBC"), MoveTemp(DisassemblyString), FGenericShaderStat::EFlags::Hidden, FShaderStatTagNames::AnalysisArtifactsName);
+				}
+			}
+
 			TRefCountPtr<ID3DBlob> CompressedData;
 
 			if (Input.Environment.CompilerFlags.Contains(CFLAG_GenerateSymbols))
diff --git a/Engine/Source/Developer/Windows/ShaderFormatD3D/Private/D3DShaderCompilerDXC.cpp b/Engine/Source/Developer/Windows/ShaderFormatD3D/Private/D3DShaderCompilerDXC.cpp
--- a/Engine/Source/Developer/Windows/ShaderFormatD3D/Private/D3DShaderCompilerDXC.cpp
+++ b/Engine/Source/Developer/Windows/ShaderFormatD3D/Private/D3DShaderCompilerDXC.cpp
@@ -1261,6 +1261,17 @@ bool CompileAndProcessD3DShaderDXC(
 			}
 
 			Output.bSucceeded = UE::ShaderCompilerCommon::ValidatePackedResourceCounts(Output, PackedResourceCounts);
+
+			// Return code reflection if requested for shader analysis
+			if (Input.Environment.CompilerFlags.Contains(CFLAG_OutputAnalysisArtifacts))
+			{
+				FGenericShaderStat ShaderCodeReflection;
+				if (CrossCompiler::FShaderConductorContext::Disassemble(CrossCompiler::EShaderConductorIR::Dxil, ShaderBlob->GetBufferPointer(), ShaderBlob->GetBufferSize(), ShaderCodeReflection))
+				{
+					ShaderCodeReflection.StatName = FName(FString::Printf(TEXT("%s (%s)"), *ShaderCodeReflection.StatName.ToString(), *Input.EntryPointName));
+					Output.ShaderStatistics.Add(MoveTemp(ShaderCodeReflection));
+				}
+			}
 		}
 
 		// Save results if compilation and reflection succeeded
diff --git a/Engine/Source/Runtime/RenderCore/Private/ShaderCore.cpp b/Engine/Source/Runtime/RenderCore/Private/ShaderCore.cpp
--- a/Engine/Source/Runtime/RenderCore/Private/ShaderCore.cpp
+++ b/Engine/Source/Runtime/RenderCore/Private/ShaderCore.cpp
@@ -53,6 +53,11 @@
 #include "Windows/HideWindowsPlatformTypes.h"
 #endif
 
+namespace FShaderStatTagNames
+{
+	extern RENDERCORE_API const FName AnalysisArtifactsName = TEXT("ShaderStatTags.AnalysisArtifacts");
+}
+
 bool operator==(const FShaderStatVariant LHS, const FShaderStatVariant RHS)
 {
 	if (LHS.IsType<bool>() && RHS.IsType<bool>())
@@ -85,16 +90,24 @@ FArchive& operator<<(FArchive& Ar, FGenericShaderStat& Stat)
 	{
 		FString StatNameString = Stat.StatName.ToString();
 		Ar << StatNameString;
+
+		FString TagNameString = Stat.TagName.ToString();
+		Ar << TagNameString;
 	}
 	else if (Ar.IsLoading())
 	{
 		FString StatNameString;
 		Ar << StatNameString;
 		Stat.StatName = FName(*StatNameString);
+
+		FString TagNameString;
+		Ar << TagNameString;
+		Stat.TagName = FName(*TagNameString);
 	}
 	else
 	{
 		Ar << Stat.StatName;
+		Ar << Stat.TagName;
 	}
 
 	Ar << Stat.Value;
@@ -104,7 +117,7 @@ FArchive& operator<<(FArchive& Ar, FGenericShaderStat& Stat)
 
 bool FGenericShaderStat::operator==(const FGenericShaderStat& RHS) const
 {
-	return (StatName == RHS.StatName) && (Value == RHS.Value) && (Flags == RHS.Flags);
+	return (StatName == RHS.StatName) && (Value == RHS.Value) && (Flags == RHS.Flags) && (TagName == RHS.TagName);
 }
 
 static TAutoConsoleVariable<bool> CVarDumpDebugInfoForCacheHits(
diff --git a/Engine/Source/Runtime/RenderCore/Public/ShaderCompilerCore.h b/Engine/Source/Runtime/RenderCore/Public/ShaderCompilerCore.h
--- a/Engine/Source/Runtime/RenderCore/Public/ShaderCompilerCore.h
+++ b/Engine/Source/Runtime/RenderCore/Public/ShaderCompilerCore.h
@@ -30,7 +30,7 @@ typedef TSharedPtr<TArray<ANSICHAR>, ESPMode::ThreadSafe> FShaderSharedAnsiStrin
 // this is for the protocol, not the data, bump if FShaderCompilerInput/FShaderPreprocessOutput serialization, SerializeWorkerInput or ProcessInputFromArchive changes.
 inline const int32 ShaderCompileWorkerInputVersion = 29;
 // this is for the protocol, not the data, bump if FShaderCompilerOutput or WriteToOutputArchive changes.
-inline const int32 ShaderCompileWorkerOutputVersion = 27;
+inline const int32 ShaderCompileWorkerOutputVersion = 28;
 // this is for the protocol, not the data.
 inline const int32 ShaderCompileWorkerSingleJobHeader = 'S';
 // this is for the protocol, not the data.
@@ -38,7 +38,7 @@ inline const int32 ShaderCompileWorkerPipelineJobHeader = 'P';
 
 // Modify this to invalidate _just_ the cache/DDC entries for individual shaders (will not cause shadermaps to rebuild if they are not otherwise out-of-date).
 // This should be bumped for changes to the FShaderCompilerOutput data structure (in addition to ShaderCompileWorkerOutputVersion)
-inline static const FGuid UE_SHADER_CACHE_VERSION = FGuid("902C35E9-5B0A-460C-96FF-C988651E9B50");
+inline static const FGuid UE_SHADER_CACHE_VERSION = FGuid("092A667D-E873-4E2D-957B-262A421AED5A");
 
 namespace UE::ShaderCompiler
 {
@@ -170,6 +170,8 @@ enum ECompilerFlags
 	CFLAG_DisableSourceStripping,
 	// Shader uses RHI Shader Binding Layout for global shader binding
 	CFLAG_ShaderBindingLayout,
+	// Request full shader analysis artifacts in output statistics. This may contain multiple compilation steps in full text form.
+	CFLAG_OutputAnalysisArtifacts,
 	CFLAG_Max,
 };
 static_assert(CFLAG_Max < 64, "Out of bitfield space! Modify FShaderCompilerFlags");
@@ -632,12 +634,13 @@ struct FShaderCompilerOutput
 	RENDERCORE_API void SerializeShaderDiagnosticData();
 
 	template<typename TValue>
-	void AddStatistic(const TCHAR* Name, TValue Value, FGenericShaderStat::EFlags Flags = FGenericShaderStat::EFlags::None)
+	void AddStatistic(const TCHAR* Name, TValue Value, FGenericShaderStat::EFlags Flags = FGenericShaderStat::EFlags::None, FName TagName = NAME_None)
 	{
 		FGenericShaderStat& Stat = ShaderStatistics.AddZeroed_GetRef();
 		Stat.StatName = FName(Name);
 		Stat.Value = FShaderStatVariant(TInPlaceType<TValue>(), Value);
 		Stat.Flags = Flags;
+		Stat.TagName = TagName;
 	}
 
 	const FShaderCodeResource& GetFinalizedCodeResource() const
diff --git a/Engine/Source/Runtime/RenderCore/Public/ShaderCore.h b/Engine/Source/Runtime/RenderCore/Public/ShaderCore.h
--- a/Engine/Source/Runtime/RenderCore/Public/ShaderCore.h
+++ b/Engine/Source/Runtime/RenderCore/Public/ShaderCore.h
@@ -58,6 +58,12 @@ using FShaderStatVariant = TVariant<bool, float, int32, uint32, FString>;
 
 bool operator==(const FShaderStatVariant LHS, const FShaderStatVariant RHS);
 
+namespace FShaderStatTagNames
+{
+	/** Tag name for shader analysis artifacts. See CFLAG_OutputAnalysisArtifacts. */
+	extern RENDERCORE_API const FName AnalysisArtifactsName;
+}
+
 struct FGenericShaderStat
 {
 public:
@@ -69,6 +75,7 @@ public:
 	FName StatName;
 	FShaderStatVariant Value;
 	EFlags Flags = EFlags::None;
+	FName TagName;
 
 	FGenericShaderStat() = default;
 
