[shaders] change a few shader formats which were unnecessarily implementing OutputDebugData to inherit from FBaseShaderFormat instead and just use the (matching) default implementation

diff --git a/Engine/Source/Developer/Apple/MetalShaderFormat/Private/MetalShaderCompiler.cpp b/Engine/Source/Developer/Apple/MetalShaderFormat/Private/MetalShaderCompiler.cpp
--- a/Engine/Source/Developer/Apple/MetalShaderFormat/Private/MetalShaderCompiler.cpp
+++ b/Engine/Source/Developer/Apple/MetalShaderFormat/Private/MetalShaderCompiler.cpp
@@ -1102,11 +1102,6 @@ void CompileMetalShader(const FShaderCompilerInput& Input, const FShaderPreproce
 	ShaderParameterParser.ValidateShaderParameterTypes(Input, bIsMobile, Output);
 }
 
-void OutputMetalDebugData(const FShaderCompilerInput& Input, const FShaderPreprocessOutput& PreprocessOutput, const FShaderCompilerOutput& Output)
-{
-	UE::ShaderCompilerCommon::DumpExtendedDebugShaderData(Input, PreprocessOutput, Output);
-}
-
 bool StripShader_Metal(TArray<uint8>& Code, class FString const& DebugPath, bool const bNative)
 {
 	bool bSuccess = false;
diff --git a/Engine/Source/Developer/Apple/MetalShaderFormat/Private/MetalShaderFormat.cpp b/Engine/Source/Developer/Apple/MetalShaderFormat/Private/MetalShaderFormat.cpp
--- a/Engine/Source/Developer/Apple/MetalShaderFormat/Private/MetalShaderFormat.cpp
+++ b/Engine/Source/Developer/Apple/MetalShaderFormat/Private/MetalShaderFormat.cpp
@@ -47,7 +47,7 @@ extern bool FinalizeLibrary_Metal(class FName const& Format, class FString const
 /** Version for shader format, this becomes part of the DDC key. */
 static const FGuid UE_SHADER_METAL_VER = FGuid("B0DC25EF-C34D-437A-94E5-5E5146AF1B9A");
 
-class FMetalShaderFormat : public IShaderFormat
+class FMetalShaderFormat : public UE::ShaderCompilerCommon::FBaseShaderFormat 
 {
 public:
 	FMetalShaderFormat()
@@ -144,14 +144,6 @@ public:
 		CompileMetalShader(Input, PreprocessOutput, Output);
 	}
 
-	virtual void OutputDebugData(
-		const FShaderCompilerInput& Input,
-		const FShaderPreprocessOutput& PreprocessOutput,
-		const FShaderCompilerOutput& Output) const override final
-	{
-		OutputMetalDebugData(Input, PreprocessOutput, Output);
-	}
-
 	virtual bool CanStripShaderCode(bool const bNativeFormat) const override final
 	{
 		return CanCompileBinaryShaders() && bNativeFormat;
diff --git a/Engine/Source/Developer/ShaderFormatVectorVM/Private/ShaderFormatVectorVM.cpp b/Engine/Source/Developer/ShaderFormatVectorVM/Private/ShaderFormatVectorVM.cpp
--- a/Engine/Source/Developer/ShaderFormatVectorVM/Private/ShaderFormatVectorVM.cpp
+++ b/Engine/Source/Developer/ShaderFormatVectorVM/Private/ShaderFormatVectorVM.cpp
@@ -20,11 +20,6 @@ extern bool CompileVectorVMShader(
 	FShaderCompilerOutput& Output,
 	const FString& WorkingDirectory);
 
-extern void OutputVectorVMDebugData(
-	const FShaderCompilerInput& Input,
-	const FShaderPreprocessOutput& PreprocessOutput,
-	const FShaderCompilerOutput& Output);
-
 static FName NAME_VVM_1_0(TEXT("VVM_1_0"));
 
 class FShaderFormatVectorVM : public UE::ShaderCompilerCommon::FBaseShaderFormat
@@ -80,11 +75,6 @@ public:
 		CompileVectorVMShader(Input, PreprocessOutput, Output, WorkingDirectory);
 	}
 
-	virtual void OutputDebugData(const FShaderCompilerInput& Input, const FShaderPreprocessOutput& PreprocessOutput, const FShaderCompilerOutput& Output) const override
-	{
-		OutputVectorVMDebugData(Input, PreprocessOutput, Output);
-	}
-
 	virtual const TCHAR* GetPlatformIncludeDirectory() const
 	{
 		return TEXT("");
diff --git a/Engine/Source/Developer/ShaderFormatVectorVM/Private/VectorVMShaderCompiler.cpp b/Engine/Source/Developer/ShaderFormatVectorVM/Private/VectorVMShaderCompiler.cpp
--- a/Engine/Source/Developer/ShaderFormatVectorVM/Private/VectorVMShaderCompiler.cpp
+++ b/Engine/Source/Developer/ShaderFormatVectorVM/Private/VectorVMShaderCompiler.cpp
@@ -135,14 +135,6 @@ bool CompileVectorVMShader(const FShaderCompilerInput& Input, const FShaderPrepr
 	return bResult;
 }
 
-void OutputVectorVMDebugData(
-	const FShaderCompilerInput& Input,
-	const FShaderPreprocessOutput& PreprocessOutput,
-	const FShaderCompilerOutput& Output)
-{
-	UE::ShaderCompilerCommon::DumpExtendedDebugShaderData(Input, PreprocessOutput, Output);
-}
-
 bool TestCompileVectorVMShader(
 	const FShaderCompilerInput& Input,
 	const FString& WorkingDirectory,
diff --git a/Engine/Source/Developer/VulkanShaderFormat/Private/VulkanShaderCompiler.cpp b/Engine/Source/Developer/VulkanShaderFormat/Private/VulkanShaderCompiler.cpp
--- a/Engine/Source/Developer/VulkanShaderFormat/Private/VulkanShaderCompiler.cpp
+++ b/Engine/Source/Developer/VulkanShaderFormat/Private/VulkanShaderCompiler.cpp
@@ -376,8 +376,3 @@ void CompileVulkanShader(const FShaderCompilerInput& Input, const FShaderPreproc
 		ensure(bSuccess);
 	}
 }
-
-void OutputVulkanDebugData(const FShaderCompilerInput& Input, const FShaderPreprocessOutput& PreprocessOutput, const FShaderCompilerOutput& Output)
-{
-	UE::ShaderCompilerCommon::DumpExtendedDebugShaderData(Input, PreprocessOutput, Output);
-}
diff --git a/Engine/Source/Developer/VulkanShaderFormat/Private/VulkanShaderFormat.cpp b/Engine/Source/Developer/VulkanShaderFormat/Private/VulkanShaderFormat.cpp
--- a/Engine/Source/Developer/VulkanShaderFormat/Private/VulkanShaderFormat.cpp
+++ b/Engine/Source/Developer/VulkanShaderFormat/Private/VulkanShaderFormat.cpp
@@ -24,11 +24,6 @@ extern void CompileVulkanShader(
 	FShaderCompilerOutput& Output,
 	const FString& WorkingDirectory);
 
-extern void OutputVulkanDebugData(
-	const FShaderCompilerInput& Input, 
-	const FShaderPreprocessOutput& PreprocessOutput, 
-	const FShaderCompilerOutput& Output);
-
 static const FGuid UE_SHADER_VULKAN_ES3_1_VER = FGuid("B84F72C8-3ECD-411E-993C-D7C7CEE26F28");
 static const FGuid UE_SHADER_VULKAN_SM5_VER = FGuid("0715D8EE-9907-4A25-93AD-A3902C8E069A");
 static const FGuid UE_SHADER_VULKAN_SM6_VER = FGuid("C5161730-83C6-40AF-A990-78CD4C1581DB");
@@ -97,11 +92,6 @@ public:
 		CompileVulkanShader(Input, PreprocessOutput, Output, WorkingDirectory);
 	}
 
-	virtual void OutputDebugData(const FShaderCompilerInput& Input, const FShaderPreprocessOutput& PreprocessOutput, const FShaderCompilerOutput& Output) const override
-	{
-		OutputVulkanDebugData(Input, PreprocessOutput, Output);
-	}
-
 	virtual const TCHAR* GetPlatformIncludeDirectory() const
 	{
 		return TEXT("Vulkan");
