[Shaders] Only serialize SPIR-V code on a successful compilation.

This fixes an assertion in the SPIR-V backend for ray-tracing shaders, which split the SPIR-V blob for each entry point.

#rnx

diff --git a/Engine/Source/Developer/VulkanShaderFormat/Private/SpirVShaderCompiler.inl b/Engine/Source/Developer/VulkanShaderFormat/Private/SpirVShaderCompiler.inl
--- a/Engine/Source/Developer/VulkanShaderFormat/Private/SpirVShaderCompiler.inl
+++ b/Engine/Source/Developer/VulkanShaderFormat/Private/SpirVShaderCompiler.inl
@@ -1666,6 +1666,7 @@ static bool CompileShaderGroup(
 	ClosestHitSerializedOutput.UsedBindlessUB.Append(IntersectionSerializedOutput.UsedBindlessUB);
 	UpdateBindlessUBs(InternalState, ClosestHitSerializedOutput, MergedOutput);
 
+	if (bSuccess)
 	{
 		// :todo-jn: Having multiple entrypoints in a single SPIRV blob crashes on FLumenHardwareRayTracingMaterialHitGroup for some reason
 		// Adjust the header before we write it out
