VulkanRHI: Add support for static uniform buffer bindings to ray gen shaders using push constants

#rnx

diff --git a/Engine/Config/VulkanPC/DataDrivenPlatformInfo.ini b/Engine/Config/VulkanPC/DataDrivenPlatformInfo.ini
--- a/Engine/Config/VulkanPC/DataDrivenPlatformInfo.ini
+++ b/Engine/Config/VulkanPC/DataDrivenPlatformInfo.ini
@@ -177,6 +177,7 @@ bSupportsRayTracingShaders = true
 bSupportsPathTracing = true
 bSupportsRayTracingCallableShaders = true
 bSupportsRayTracingProceduralPrimitive = true
+StaticShaderBindingLayoutSupport=RayTracingOnly
 
 bSupportsGPUScene = true
 bSupportsPrimitiveShaders = false
diff --git a/Engine/Shaders/Public/Platform/Vulkan/VulkanCommon.ush b/Engine/Shaders/Public/Platform/Vulkan/VulkanCommon.ush
--- a/Engine/Shaders/Public/Platform/Vulkan/VulkanCommon.ush
+++ b/Engine/Shaders/Public/Platform/Vulkan/VulkanCommon.ush
@@ -7,7 +7,7 @@
 #pragma once
 
 // Update this GUID to invalidate shader recompilation for only Vulkan shaders
-#pragma message("UESHADERMETADATA_VERSION C77F7A23-39F4-49D6-94C7-AEE00F7EF98A")
+#pragma message("UESHADERMETADATA_VERSION E8F4EF26-680E-440C-980A-170F9B9060F0")
 
 #define COMPILER_SUPPORTS_ULONG_TYPES 1
 #define COMPILER_SUPPORTS_UINT64_IMAGE_ATOMICS (1 && COMPILER_SUPPORTS_ULONG_TYPES)
@@ -157,3 +157,13 @@ void ImageInterlockedOrUInt64(RWTexture2D<UlongType> Dest, uint2 Coord, UlongTyp
 #endif
 
 #define GetGlobalSampler(Filter, WrapMode) VulkanGlobal##Filter##WrapMode##Sampler
+
+
+#if UE_BINDLESS_ENABLED && RAYGENSHADER
+// Push constants are currently only used for ray tracing's static UBs (Scene, View, NaniteRayTracing, LumenHardwareRayTracing)
+struct FVulkanPushConstants
+{
+	uint StaticUniformBufferBindings[8];  // keep in sync with FRHIShaderBindingLayout::MaxUniformBufferEntries
+};
+[[vk::push_constant]] ConstantBuffer<FVulkanPushConstants> VulkanPushConstants;
+#endif
\ No newline at end of file
diff --git a/Engine/Source/Developer/VulkanShaderFormat/Private/SpirVShaderCompiler.inl b/Engine/Source/Developer/VulkanShaderFormat/Private/SpirVShaderCompiler.inl
--- a/Engine/Source/Developer/VulkanShaderFormat/Private/SpirVShaderCompiler.inl
+++ b/Engine/Source/Developer/VulkanShaderFormat/Private/SpirVShaderCompiler.inl
@@ -14,6 +14,7 @@ public:
 		: Input(InInput)
 		, ParameterParser(InParameterParser)
 		, bUseBindlessUniformBuffer(InInput.IsRayTracingShader() && ((EShaderFrequency)InInput.Target.Frequency != SF_RayGen))
+		, bUseStaticUniformBufferBindings((EShaderFrequency)InInput.Target.Frequency == SF_RayGen)
 		, bIsRayHitGroupShader(InInput.IsRayTracingShader() && ((EShaderFrequency)InInput.Target.Frequency == SF_RayHitGroup))
 		, bSupportsBindless(InInput.IsBindlessEnabled())
 		, bDebugDump(InInput.DumpDebugInfoEnabled())
@@ -33,6 +34,7 @@ public:
 	const FShaderParameterParser* ParameterParser;
 
 	const bool bUseBindlessUniformBuffer;
+	const bool bUseStaticUniformBufferBindings;
 	const bool bIsRayHitGroupShader;
 
 	const bool bSupportsBindless;
@@ -52,6 +54,7 @@ public:
 	FString IntersectionEntry;
 
 	TArray<FString> AllBindlessUBs;
+	TArray<FString> PushConstantUBs;
 	uint32 ShaderRecordGlobalsSize = 0;
 
 	// Forwarded calls for convenience
@@ -596,11 +599,14 @@ static void GatherSpirvReflectionBindings(
 				const FString BindingName(ANSI_TO_TCHAR(pBinding->name));
 				if (BindingName.StartsWith(kBindlessCBPrefix))
 				{
-					check(InternalState.bUseBindlessUniformBuffer);
+					check(InternalState.bUseBindlessUniformBuffer || InternalState.bUseStaticUniformBufferBindings);
 					Reflection.ChangeDescriptorBindingNumbers(pBinding, 0, VulkanBindless::BindlessUniformBufferSet);
 					const FString BindlessUBName = GetBindlessUBNameFromHeap(BindingName);
-					checkf(InternalState.AllBindlessUBs.Contains(BindlessUBName), TEXT("Bindless Uniform Buffer was found in SPIRV but not tracked in internal state"));
-					OutBindlessUB.Add(BindlessUBName);
+					if (!InternalState.PushConstantUBs.Contains(BindlessUBName))
+					{
+						checkf(InternalState.AllBindlessUBs.Contains(BindlessUBName), TEXT("Bindless Uniform Buffer was found in SPIRV but not tracked in internal state"));
+						OutBindlessUB.Add(BindlessUBName);
+					}
 					OutBindings.UniformBuffers.RemoveAtSwap(Index);
 				}
 				else
@@ -1453,6 +1459,9 @@ static void ModifyCompilerInput(FSpirvShaderCompilerInternalState& InternalState
 		// Name of the structure in raytracing shader records in VulkanCommon.usf
 		Input.RequiredSymbols.Add(TEXT("HitGroupSystemRootConstants"));
 
+		// Name of the structure for static uniform buffers in VulkanCommon.usf
+		Input.RequiredSymbols.Add(TEXT("VulkanPushConstants"));
+
 		// Always remove dead code for ray tracing shaders regardless of cvar settings, 
 		// we can't support multiple entrypoints remaining in the binaries
 		Input.Environment.CompilerFlags.Add(CFLAG_RemoveDeadCode);
@@ -1479,17 +1488,29 @@ static void UpdateBindlessUBs(const FSpirvShaderCompilerInternalState& InternalS
 	}
 }
 
+enum class EBindlessUniformBufferType
+{
+	ShaderRecord,
+	PushConstant
+};
+
 // :todo-jn: TEMPORARY EXPERIMENT - will eventually move into preprocessing step
-static TArray<FString> ConvertUBToBindless(FString& PreprocessedShaderSource)
+static TArray<FString> ConvertUBToBindless(FString& PreprocessedShaderSource, 
+	EBindlessUniformBufferType BindlessUniformBufferType = EBindlessUniformBufferType::ShaderRecord, 
+	TArrayView<FString> UBNames = {})
 {
 	// Fill a map so we pull our bindless sampler/resource indices from the right struct
 	// :todo-jn: Do we not have the layout somewhere instead of calculating offsets?  there must be a better way...
-	auto GenerateNewDecl = [](const int32 CBIndex, const FString& Members, const FString& CBName)
+	auto GenerateNewDecl = [BindlessUniformBufferType](const int32 CBIndex, const FString& Members, const FString& CBName)
 	{
 		const FString PrefixedCBName = FString::Printf(TEXT("%s%d_%s"), *SpirvShaderCompiler::kBindlessCBPrefix, CBIndex, *CBName);
 		const FString BindlessCBType = PrefixedCBName + TEXT("_Type");
 		const FString BindlessCBHeapName = PrefixedCBName + SpirvShaderCompiler::kBindlessHeapSuffix;
 		const FString PaddingName = FString::Printf(TEXT("%s_Padding"), *CBName);
+		const TCHAR* BindlessIndexName = 
+			(BindlessUniformBufferType == EBindlessUniformBufferType::ShaderRecord) ?
+			TEXT("VulkanHitGroupSystemParameters.BindlessUniformBuffers") :
+			TEXT("VulkanPushConstants.StaticUniformBufferBindings");
 
 		FString CBDecl;
 		CBDecl.Reserve(Members.Len() * 3);  // start somewhere approx less bad
@@ -1501,8 +1522,8 @@ static TArray<FString> ConvertUBToBindless(FString& PreprocessedShaderSource)
 		CBDecl += FString::Printf(TEXT("ConstantBuffer<%s> %s[];\n"), *BindlessCBType, *BindlessCBHeapName);
 
 		// Now bring in the CB
-		CBDecl += FString::Printf(TEXT("static const %s %s = %s[VulkanHitGroupSystemParameters.BindlessUniformBuffers[%d]];\n"),
-			*BindlessCBType, *PrefixedCBName, *BindlessCBHeapName, CBIndex);
+		CBDecl += FString::Printf(TEXT("static const %s %s = %s[%s[%d]];\n"),
+			*BindlessCBType, *PrefixedCBName, *BindlessCBHeapName, BindlessIndexName, CBIndex);
 
 		// Now create a global scope var for each value (as the cbuffer would provide) to patch in seemlessly with the rest of the code
 		uint32 MemberOffset = 0;
@@ -1564,6 +1585,11 @@ static TArray<FString> ConvertUBToBindless(FString& PreprocessedShaderSource)
 
 	// replace "cbuffer" decl with a struct filled from bindless constant buffer
 	TArray<FString> BindlessUBs;
+	if (!UBNames.IsEmpty())
+	{
+		BindlessUBs.SetNum(UBNames.Num());
+	}
+
 	{
 		const FString UniformBufferDeclIdentifier = TEXT("cbuffer");
 
@@ -1574,23 +1600,41 @@ static TArray<FString> ConvertUBToBindless(FString& PreprocessedShaderSource)
 			const TCHAR* StructNameEndPtr = ParseHLSLSymbolName(&PreprocessedShaderSource[SearchIndex + UniformBufferDeclIdentifier.Len()], StructName);
 			check(StructName.Len() > 0);
 
-			const int32 CBIndex = BindlessUBs.Add(StructName);
-			check(CBIndex < 16);
-
 			const TCHAR* OpeningBracePtr = FCString::Strstr(&PreprocessedShaderSource[SearchIndex + UniformBufferDeclIdentifier.Len()], TEXT("{"));
 			check(OpeningBracePtr);
 			const TCHAR* ClosingBracePtr = FindMatchingClosingBrace(OpeningBracePtr + 1);
 			check(ClosingBracePtr);
 			const int32 ClosingBraceIndex = ClosingBracePtr - (*PreprocessedShaderSource);
 
-			const FString Members = FString::ConstructFromPtrSize(OpeningBracePtr + 1, ClosingBracePtr - OpeningBracePtr - 1);
-			const FString NewDecl = GenerateNewDecl(CBIndex, Members, StructName);
+			int32 CBIndex;
+			if (UBNames.IsEmpty())
+			{
+				CBIndex = BindlessUBs.Add(StructName);
+				check(CBIndex < 16);
+			}
+			else
+			{
+				CBIndex = UBNames.Find(*StructName);
+				if (CBIndex != INDEX_NONE)
+				{
+					BindlessUBs[CBIndex] = StructName;
+				}
+			}
+
+			int32 NextStartIndex = ClosingBraceIndex;
+			if (CBIndex != INDEX_NONE)
+			{
+				const FString Members = FString::ConstructFromPtrSize(OpeningBracePtr + 1, ClosingBracePtr - OpeningBracePtr - 1);
+				const FString NewDecl = GenerateNewDecl(CBIndex, Members, StructName);
 
-			const int32 OldDeclLen = ClosingBraceIndex - SearchIndex + 1;
-			PreprocessedShaderSource.RemoveAt(SearchIndex, OldDeclLen, EAllowShrinking::No);
-			PreprocessedShaderSource.InsertAt(SearchIndex, NewDecl);
+				const int32 OldDeclLen = ClosingBraceIndex - SearchIndex + 1;
+				PreprocessedShaderSource.RemoveAt(SearchIndex, OldDeclLen, EAllowShrinking::No);
+				PreprocessedShaderSource.InsertAt(SearchIndex, NewDecl);
+
+				NextStartIndex = SearchIndex + NewDecl.Len();
+			}
 
-			SearchIndex = PreprocessedShaderSource.Find(UniformBufferDeclIdentifier, ESearchCase::CaseSensitive, ESearchDir::FromStart, SearchIndex + NewDecl.Len());
+			SearchIndex = PreprocessedShaderSource.Find(UniformBufferDeclIdentifier, ESearchCase::CaseSensitive, ESearchDir::FromStart, NextStartIndex);
 		}
 	}
 	return BindlessUBs;
diff --git a/Engine/Source/Developer/VulkanShaderFormat/Private/VulkanShaderCompiler.cpp b/Engine/Source/Developer/VulkanShaderFormat/Private/VulkanShaderCompiler.cpp
--- a/Engine/Source/Developer/VulkanShaderFormat/Private/VulkanShaderCompiler.cpp
+++ b/Engine/Source/Developer/VulkanShaderFormat/Private/VulkanShaderCompiler.cpp
@@ -300,13 +300,25 @@ void CompileVulkanShader(const FShaderCompilerInput& Input, const FShaderPreproc
 
 	FVulkanShaderCompilerInternalState InternalState(Input, &ShaderParameterParser);
 
+	if (InternalState.bUseStaticUniformBufferBindings)
+	{
+		// Match uniform buffers listed in FRaytracingShaderBindingLayout
+		TArray<FString> RayTracingStaticUniformBufferNames;
+		RayTracingStaticUniformBufferNames.Add(TEXT("Scene"));
+		RayTracingStaticUniformBufferNames.Add(TEXT("View"));
+		RayTracingStaticUniformBufferNames.Add(TEXT("NaniteRayTracing"));
+		RayTracingStaticUniformBufferNames.Add(TEXT("LumenHardwareRayTracingUniformBuffer"));
+
+		InternalState.PushConstantUBs = SpirvShaderCompiler::ConvertUBToBindless(PreprocessedSource, SpirvShaderCompiler::EBindlessUniformBufferType::PushConstant, RayTracingStaticUniformBufferNames);
+	}
+
 	if (InternalState.bUseBindlessUniformBuffer)
 	{
 		InternalState.ShaderRecordGlobalsSize = ConvertGlobalsToShaderRecord(ShaderParameterParser, ReplacedGlobals, PreprocessedSource, Output);
-		InternalState.AllBindlessUBs = SpirvShaderCompiler::ConvertUBToBindless(PreprocessedSource);
+		InternalState.AllBindlessUBs = SpirvShaderCompiler::ConvertUBToBindless(PreprocessedSource, SpirvShaderCompiler::EBindlessUniformBufferType::ShaderRecord);
 	}
 
-	if (ShaderParameterParser.DidModifyShader() || InternalState.AllBindlessUBs.Num() || InternalState.ShaderRecordGlobalsSize)
+	if (ShaderParameterParser.DidModifyShader() || InternalState.AllBindlessUBs.Num() || InternalState.ShaderRecordGlobalsSize || InternalState.PushConstantUBs.Num())
 	{
 		Output.ModifiedShaderSource = PreprocessedSource;
 	}
diff --git a/Engine/Source/Runtime/VulkanRHI/Private/VulkanBindlessDescriptorManager.cpp b/Engine/Source/Runtime/VulkanRHI/Private/VulkanBindlessDescriptorManager.cpp
--- a/Engine/Source/Runtime/VulkanRHI/Private/VulkanBindlessDescriptorManager.cpp
+++ b/Engine/Source/Runtime/VulkanRHI/Private/VulkanBindlessDescriptorManager.cpp
@@ -635,6 +635,20 @@ void FVulkanBindlessDescriptorManager::Init()
 		ZeroVulkanStruct(PipelineLayoutCreateInfo, VK_STRUCTURE_TYPE_PIPELINE_LAYOUT_CREATE_INFO);
 		PipelineLayoutCreateInfo.setLayoutCount = VulkanBindless::MaxNumSets;
 		PipelineLayoutCreateInfo.pSetLayouts = DescriptorSetLayouts;
+
+		// Set the static uniform buffers with push constants (only enabled for ray tracing in DDPI, only used for RGS)
+		VkPushConstantRange PushConstants;
+		if (RHIGetStaticShaderBindingLayoutSupport(GMaxRHIShaderPlatform) != ERHIStaticShaderBindingLayoutSupport::Unsupported)
+		{
+			checkf(RHIGetStaticShaderBindingLayoutSupport(GMaxRHIShaderPlatform) == ERHIStaticShaderBindingLayoutSupport::RayTracingOnly,
+				TEXT("StaticShaderBindingLayoutSupport is currently only supported for ray gen shaders in Vulkan."));
+			PushConstants.offset = 0;
+			PushConstants.size = FRHIShaderBindingLayout::MaxUniformBufferEntries * sizeof(uint32);
+			PushConstants.stageFlags = VK_SHADER_STAGE_RAYGEN_BIT_KHR;
+			PipelineLayoutCreateInfo.pushConstantRangeCount = 1;
+			PipelineLayoutCreateInfo.pPushConstantRanges = &PushConstants;
+		}
+
 		VERIFYVULKANRESULT(VulkanRHI::vkCreatePipelineLayout(DeviceHandle, &PipelineLayoutCreateInfo, VULKAN_CPU_ALLOCATOR, &BindlessPipelineLayout));
 		VULKAN_SET_DEBUG_NAME(Device, VK_OBJECT_TYPE_PIPELINE_LAYOUT, BindlessPipelineLayout, TEXT("BindlessPipelineLayout(SetCount=%d)"), VulkanBindless::MaxNumSets);
 	}
diff --git a/Engine/Source/Runtime/VulkanRHI/Private/VulkanCommands.cpp b/Engine/Source/Runtime/VulkanRHI/Private/VulkanCommands.cpp
--- a/Engine/Source/Runtime/VulkanRHI/Private/VulkanCommands.cpp
+++ b/Engine/Source/Runtime/VulkanRHI/Private/VulkanCommands.cpp
@@ -455,9 +455,26 @@ void FVulkanCommandListContext::RHISetStaticUniformBuffers(const FUniformBufferS
 {
 	FMemory::Memzero(GlobalUniformBuffers.GetData(), GlobalUniformBuffers.Num() * sizeof(FRHIUniformBuffer*));
 
-	for (int32 Index = 0; Index < InUniformBuffers.GetUniformBufferCount(); ++Index)
+	if (const FRHIShaderBindingLayout* Layout = InUniformBuffers.GetShaderBindingLayout())
 	{
-		GlobalUniformBuffers[InUniformBuffers.GetSlot(Index)] = InUniformBuffers.GetUniformBuffer(Index);
+		check(InUniformBuffers.GetUniformBufferCount() == Layout->GetNumUniformBufferEntries());
+
+		for (int32 Index = 0; Index < InUniformBuffers.GetUniformBufferCount(); ++Index)
+		{
+			GlobalUniformBuffers[Index] = InUniformBuffers.GetUniformBuffer(Index);
+			checkf(GlobalUniformBuffers[Index], TEXT("Static uniform buffer at index %d is referenced in the shader binding layout but is not provided"), Index);
+		}
+
+		ShaderBindingLayout = Layout;
+	}
+	else
+	{
+		for (int32 Index = 0; Index < InUniformBuffers.GetUniformBufferCount(); ++Index)
+		{
+			GlobalUniformBuffers[InUniformBuffers.GetSlot(Index)] = InUniformBuffers.GetUniformBuffer(Index);
+		}
+
+		ShaderBindingLayout = nullptr;
 	}
 }
 
diff --git a/Engine/Source/Runtime/VulkanRHI/Private/VulkanContext.cpp b/Engine/Source/Runtime/VulkanRHI/Private/VulkanContext.cpp
--- a/Engine/Source/Runtime/VulkanRHI/Private/VulkanContext.cpp
+++ b/Engine/Source/Runtime/VulkanRHI/Private/VulkanContext.cpp
@@ -10,6 +10,7 @@
 #include "VulkanDevice.h"
 #include "VulkanPendingState.h"
 #include "VulkanQuery.h"
+#include "VulkanBindlessDescriptorManager.h"
 #include "DynamicRHI.h"
 
 
@@ -711,6 +712,30 @@ bool FVulkanCommandListContext::AcquirePoolSetAndDescriptorsIfNeeded(const FVulk
 	return false;
 }
 
+void FVulkanCommandListContext::ApplyShaderBindingLayout(VkShaderStageFlags ShaderStageFlags)
+{
+	if (ShaderBindingLayout && ShaderBindingLayout->GetNumUniformBufferEntries())
+	{
+		checkf(Device.GetBindlessDescriptorManager()->IsSupported(), TEXT("Static uniform buffer require bindless."));
+
+		uint32 UBBindlessIndices[FRHIShaderBindingLayout::MaxUniformBufferEntries];
+		for (uint32 Index = 0; Index < ShaderBindingLayout->GetNumUniformBufferEntries(); ++Index)
+		{
+			checkf(GlobalUniformBuffers[Index], TEXT("Missing static uniform buffer for shader binding layout (%s)."),
+				*ShaderBindingLayout->GetUniformBufferEntry(Index).LayoutName);
+
+			FVulkanUniformBuffer* StaticUniformBuffer = ResourceCast(GlobalUniformBuffers[Index]);
+			const FRHIDescriptorHandle BindlessHandle = StaticUniformBuffer->GetBindlessHandle();
+			check(BindlessHandle.IsValid());
+
+			UBBindlessIndices[Index] = BindlessHandle.GetIndex();
+		}
+
+		VulkanRHI::vkCmdPushConstants( GetCommandBuffer().GetHandle(), Device.GetBindlessDescriptorManager()->GetPipelineLayout(),
+			ShaderStageFlags, 0 /*offset*/, sizeof(uint32) * ShaderBindingLayout->GetNumUniformBufferEntries(), UBBindlessIndices);
+	}
+}
+
 
 
 
diff --git a/Engine/Source/Runtime/VulkanRHI/Private/VulkanContext.h b/Engine/Source/Runtime/VulkanRHI/Private/VulkanContext.h
--- a/Engine/Source/Runtime/VulkanRHI/Private/VulkanContext.h
+++ b/Engine/Source/Runtime/VulkanRHI/Private/VulkanContext.h
@@ -415,6 +415,15 @@ public:
 
 	bool AcquirePoolSetAndDescriptorsIfNeeded(const class FVulkanDescriptorSetsLayout& Layout, bool bNeedDescriptors, VkDescriptorSet* OutDescriptors);
 
+	bool IsSwapchainImage(FRHITexture* InTexture) const;
+	VkFormat GetSwapchainImageFormat() const;
+	FVulkanSwapChain* GetSwapChain() const;
+
+	FVulkanRenderPass* PrepareRenderPassForPSOCreation(const FGraphicsPipelineStateInitializer& Initializer);
+	FVulkanRenderPass* PrepareRenderPassForPSOCreation(const FVulkanRenderTargetLayout& Initializer);
+
+	void ApplyShaderBindingLayout(VkShaderStageFlags ShaderStageFlags);
+
 protected:
 	void ExtractDepthStencilLayouts(const FRHIRenderPassInfo& InInfo, VkImageLayout& OutDepthLayout, VkImageLayout& OutStencilLayout);
 	void FillDynamicRenderingInfo(const FRHIRenderPassInfo& RenderPassInfo, FVulkanDynamicRenderingInfo& OutRenderingInfo);
@@ -450,14 +459,6 @@ protected:
 	TMap<uint32, class FVulkanTypedDescriptorPoolSet*> TypedDescriptorPoolSets;
 	void AcquirePoolSetContainer();
 
-public:
-	bool IsSwapchainImage(FRHITexture* InTexture) const;
-	VkFormat GetSwapchainImageFormat() const;
-	FVulkanSwapChain* GetSwapChain() const;
-
-	FVulkanRenderPass* PrepareRenderPassForPSOCreation(const FGraphicsPipelineStateInitializer& Initializer);
-	FVulkanRenderPass* PrepareRenderPassForPSOCreation(const FVulkanRenderTargetLayout& Initializer);
-
 private:
 	void RHIClearMRT(bool bClearColor, int32 NumClearColors, const FLinearColor* ColorArray, bool bClearDepth, float Depth, bool bClearStencil, uint32 Stencil);
 
@@ -476,6 +477,7 @@ private:
 	void ApplyStaticUniformBuffers(TRHIShader* Shader);
 
 	TArray<FRHIUniformBuffer*> GlobalUniformBuffers;
+	const FRHIShaderBindingLayout* ShaderBindingLayout = nullptr;
 };
 
 class FVulkanCommandListContextImmediate : public FVulkanCommandListContext
diff --git a/Engine/Source/Runtime/VulkanRHI/Private/VulkanRayTracing.cpp b/Engine/Source/Runtime/VulkanRHI/Private/VulkanRayTracing.cpp
--- a/Engine/Source/Runtime/VulkanRHI/Private/VulkanRayTracing.cpp
+++ b/Engine/Source/Runtime/VulkanRHI/Private/VulkanRayTracing.cpp
@@ -1945,6 +1945,9 @@ static void SetRayGenResources(FVulkanDevice& Device, FVulkanCommandListContext&
 		Device.GetBindlessDescriptorManager()->RegisterUniformBuffers(Context, VK_PIPELINE_BIND_POINT_RAY_TRACING_KHR, StageUBs);
 	}
 
+	// Static uniform buffers
+	Context.ApplyShaderBindingLayout(VK_SHADER_STAGE_RAYGEN_BIT_KHR);
+
 	// Add all the UBs references by the shader table
 	TArrayView<TRefCountPtr<FRHIUniformBuffer>> ShaderTableUBs = ShaderTable->GetUBRefs();
 	for (TRefCountPtr<FRHIUniformBuffer>& UniformBuffer : ShaderTableUBs)
