VulkanRHI: Fix loose data in the shader record when multiple params are used

#rnx

diff --git a/Engine/Shaders/Public/Platform/Vulkan/VulkanCommon.ush b/Engine/Shaders/Public/Platform/Vulkan/VulkanCommon.ush
--- a/Engine/Shaders/Public/Platform/Vulkan/VulkanCommon.ush
+++ b/Engine/Shaders/Public/Platform/Vulkan/VulkanCommon.ush
@@ -7,7 +7,7 @@
 #pragma once
 
 // Update this GUID to invalidate shader recompilation for only Vulkan shaders
-#pragma message("UESHADERMETADATA_VERSION E8F4EF26-680E-440C-980A-170F9B9060F0")
+#pragma message("UESHADERMETADATA_VERSION ABC4AA8D-98C2-48FA-9403-341F12D67621")
 
 #define COMPILER_SUPPORTS_ULONG_TYPES 1
 #define COMPILER_SUPPORTS_UINT64_IMAGE_ATOMICS (1 && COMPILER_SUPPORTS_ULONG_TYPES)
diff --git a/Engine/Source/Developer/VulkanShaderFormat/Private/VulkanShaderCompiler.cpp b/Engine/Source/Developer/VulkanShaderFormat/Private/VulkanShaderCompiler.cpp
--- a/Engine/Source/Developer/VulkanShaderFormat/Private/VulkanShaderCompiler.cpp
+++ b/Engine/Source/Developer/VulkanShaderFormat/Private/VulkanShaderCompiler.cpp
@@ -190,9 +190,9 @@ static uint32 ConvertGlobalsToShaderRecord(const FShaderParameterParser& ShaderP
 
 		HandleReflectedGlobalConstantBufferMember(
 			ParamName,
-			ShaderRecordParamCount++,
-			ShaderRecordGlobalsSize,
-			ParamSize,
+			0 /*ConstantBufferIndex*/,
+			ShaderRecordGlobalsSize /*ReflectionOffset*/,
+			ParamSize /*ReflectionSize*/,
 			Output
 		);
 
