[Shaders] Output statistics of used resources in Vulkan shader backend to match the D3D backend.

- Emit shader statistics for Vulkan backend w/ and w/o bindless.
- Improve formating of statistics by moving newline character at the beginning of each entry,
this way there's no empty line at the end and "Stats:" is on its own line, which is otherwise out of place with respect to the remaining lines.
- Also add RW resources to D3D statistics.

#test Tested in UE project with D3D12 and Vulkan w/ and w/o RayTracing and Bindless setting.
#rnx

diff --git a/Engine/Source/Developer/ShaderCompilerCommon/Public/SpirvReflectCommon.h b/Engine/Source/Developer/ShaderCompilerCommon/Public/SpirvReflectCommon.h
--- a/Engine/Source/Developer/ShaderCompilerCommon/Public/SpirvReflectCommon.h
+++ b/Engine/Source/Developer/ShaderCompilerCommon/Public/SpirvReflectCommon.h
@@ -31,6 +31,8 @@ struct FSpirvReflectBindings
 	TArray<SpvReflectDescriptorBinding*> SBufferSRVs;
 	TArray<SpvReflectDescriptorBinding*> SBufferUAVs;
 	TArray<SpvReflectDescriptorBinding*> AccelerationStructures;
+	int32 NumBindlessSamplers = 0;
+	int32 NumBindlessResources = 0;
 
 	/** Adds the specified descriptor binding to the corresponding container iff the descriptor is marked as being accessed. */
 	SHADERCOMPILERCOMMON_API void AddDescriptorBinding(SpvReflectDescriptorBinding* InBinding);
diff --git a/Engine/Source/Developer/VulkanShaderFormat/Private/SpirVShaderCompiler.inl b/Engine/Source/Developer/VulkanShaderFormat/Private/SpirVShaderCompiler.inl
--- a/Engine/Source/Developer/VulkanShaderFormat/Private/SpirVShaderCompiler.inl
+++ b/Engine/Source/Developer/VulkanShaderFormat/Private/SpirVShaderCompiler.inl
@@ -562,8 +562,10 @@ static void GatherSpirvReflectionBindings(
 	if (InternalState.bSupportsBindless)
 	{
 		// Move the bindless heap to its dedicated descriptor set and remove it from our regular binding arrays
-		auto MoveBindlessHeaps = [&](TArray<SpvReflectDescriptorBinding*>& BindingArray, const TCHAR* HeapPrefix, uint32 BinldessDescSetNo)
+		auto MoveBindlessHeaps = [&](TArray<SpvReflectDescriptorBinding*>& BindingArray, const TCHAR* HeapPrefix, uint32 BinldessDescSetNo, int32& InOutNumBindlessEntries)
 		{
+			InOutNumBindlessEntries += BindingArray.Num();
+
 			for (int32 Index = BindingArray.Num() - 1; Index >= 0; --Index)
 			{
 				const SpvReflectDescriptorBinding* pBinding = BindingArray[Index];
@@ -578,17 +580,17 @@ static void GatherSpirvReflectionBindings(
 		};
 
 		// Remove sampler heaps from binding arrays
-		MoveBindlessHeaps(OutBindings.Samplers, FShaderParameterParser::kBindlessSamplerArrayPrefix, VulkanBindless::BindlessSamplerSet);
+		MoveBindlessHeaps(OutBindings.Samplers, FShaderParameterParser::kBindlessSamplerArrayPrefix, VulkanBindless::BindlessSamplerSet, OutBindings.NumBindlessSamplers);
 
 		// Remove resource heaps from binding arrays
-		MoveBindlessHeaps(OutBindings.SBufferUAVs, FShaderParameterParser::kBindlessUAVArrayPrefix, VulkanBindless::BindlessStorageBufferSet);
-		MoveBindlessHeaps(OutBindings.SBufferUAVs, FShaderParameterParser::kBindlessSRVArrayPrefix, VulkanBindless::BindlessStorageBufferSet);  // try with both prefixes, they were merged earlier
-		MoveBindlessHeaps(OutBindings.TextureSRVs, FShaderParameterParser::kBindlessSRVArrayPrefix, VulkanBindless::BindlessSampledImageSet);
-		MoveBindlessHeaps(OutBindings.TextureUAVs, FShaderParameterParser::kBindlessUAVArrayPrefix, VulkanBindless::BindlessStorageImageSet);
-		MoveBindlessHeaps(OutBindings.TextureUAVs, FShaderParameterParser::kBindlessSRVArrayPrefix, VulkanBindless::BindlessStorageImageSet);  // try with both prefixes, R64 SRV textures are read as storage images
-		MoveBindlessHeaps(OutBindings.TBufferSRVs, FShaderParameterParser::kBindlessSRVArrayPrefix, VulkanBindless::BindlessUniformTexelBufferSet);
-		MoveBindlessHeaps(OutBindings.TBufferUAVs, FShaderParameterParser::kBindlessUAVArrayPrefix, VulkanBindless::BindlessStorageTexelBufferSet);
-		MoveBindlessHeaps(OutBindings.AccelerationStructures, FShaderParameterParser::kBindlessSRVArrayPrefix, VulkanBindless::BindlessAccelerationStructureSet);
+		MoveBindlessHeaps(OutBindings.SBufferUAVs, FShaderParameterParser::kBindlessUAVArrayPrefix, VulkanBindless::BindlessStorageBufferSet, OutBindings.NumBindlessResources);
+		MoveBindlessHeaps(OutBindings.SBufferUAVs, FShaderParameterParser::kBindlessSRVArrayPrefix, VulkanBindless::BindlessStorageBufferSet, OutBindings.NumBindlessResources);  // try with both prefixes, they were merged earlier
+		MoveBindlessHeaps(OutBindings.TextureSRVs, FShaderParameterParser::kBindlessSRVArrayPrefix, VulkanBindless::BindlessSampledImageSet, OutBindings.NumBindlessResources);
+		MoveBindlessHeaps(OutBindings.TextureUAVs, FShaderParameterParser::kBindlessUAVArrayPrefix, VulkanBindless::BindlessStorageImageSet, OutBindings.NumBindlessResources);
+		MoveBindlessHeaps(OutBindings.TextureUAVs, FShaderParameterParser::kBindlessSRVArrayPrefix, VulkanBindless::BindlessStorageImageSet, OutBindings.NumBindlessResources);  // try with both prefixes, R64 SRV textures are read as storage images
+		MoveBindlessHeaps(OutBindings.TBufferSRVs, FShaderParameterParser::kBindlessSRVArrayPrefix, VulkanBindless::BindlessUniformTexelBufferSet, OutBindings.NumBindlessResources);
+		MoveBindlessHeaps(OutBindings.TBufferUAVs, FShaderParameterParser::kBindlessUAVArrayPrefix, VulkanBindless::BindlessStorageTexelBufferSet, OutBindings.NumBindlessResources);
+		MoveBindlessHeaps(OutBindings.AccelerationStructures, FShaderParameterParser::kBindlessSRVArrayPrefix, VulkanBindless::BindlessAccelerationStructureSet, OutBindings.NumBindlessResources);
 
 		// Move uniform buffers to the correct set
 		{
@@ -911,6 +913,19 @@ static bool BuildShaderOutputFromSpirv(
 		Output.NumTextureSamplers = AddReflectionInfos(Bindings.Samplers, VK_DESCRIPTOR_TYPE_SAMPLER, 0);
 		SerializedOutput.PackedResourceCounts.NumSamplers = (uint8)Output.NumTextureSamplers;
 
+		// Output resource statistics. This matches the statistics from the D3D backend, except that the limit is not a fixed number and runtime dependent.
+		if (InternalState.bSupportsBindless)
+		{
+			Output.AddStatistic<uint32>(TEXT("Bindless Resources"), Bindings.NumBindlessResources);
+			Output.AddStatistic<uint32>(TEXT("Bindless Samplers"), Bindings.NumBindlessSamplers);
+		}
+		else
+		{
+			Output.AddStatistic<uint32>(TEXT("Resources Used"), (uint32)SRVBindings);
+			Output.AddStatistic<uint32>(TEXT("RW Resources Used"), (uint32)UAVBindings);
+			Output.AddStatistic<uint32>(TEXT("Samplers Used"), Output.NumTextureSamplers);
+		}
+
 		AddReflectionInfos(Bindings.AccelerationStructures, VK_DESCRIPTOR_TYPE_ACCELERATION_STRUCTURE_KHR, 0);
 	}
 
diff --git a/Engine/Source/Developer/Windows/ShaderFormatD3D/Private/D3DShaderCompiler.inl b/Engine/Source/Developer/Windows/ShaderFormatD3D/Private/D3DShaderCompiler.inl
--- a/Engine/Source/Developer/Windows/ShaderFormatD3D/Private/D3DShaderCompiler.inl
+++ b/Engine/Source/Developer/Windows/ShaderFormatD3D/Private/D3DShaderCompiler.inl
@@ -564,6 +564,9 @@ inline void GenerateFinalOutput(
 			Output.AddStatistic<uint32>(TEXT("Resources Used"), CompileData.NumSRVs);
 			Output.AddStatistic<uint32>(TEXT("Resource Limit"), CompileData.MaxSRVs);
 
+			Output.AddStatistic<uint32>(TEXT("RW Resources Used"), CompileData.NumUAVs);
+			Output.AddStatistic<uint32>(TEXT("RW Resource Limit"), CompileData.MaxUAVs);
+
 			Output.AddStatistic<uint32>(TEXT("Samplers Used"), CompileData.NumSamplers);
 			Output.AddStatistic<uint32>(TEXT("Sampler Limit"), CompileData.MaxSamplers);
 		}
diff --git a/Engine/Source/Editor/MaterialEditor/Private/MaterialStatsCommon.cpp b/Engine/Source/Editor/MaterialEditor/Private/MaterialStatsCommon.cpp
--- a/Engine/Source/Editor/MaterialEditor/Private/MaterialStatsCommon.cpp
+++ b/Engine/Source/Editor/MaterialEditor/Private/MaterialStatsCommon.cpp
@@ -434,15 +434,16 @@ void FMaterialStatsUtils::GetRepresentativeShaderTypesAndDescriptions(TMap<FName
 }
 
 #if WITH_EDITORONLY_DATA
-static FString GetShaderString(const TArray<FGenericShaderStat>& Statistics)
+// Returns the specified array of shader statistics as string. Each entry begins with a newline character '\n'.
+static FString GenericShaderStatsToString(const TArray<FGenericShaderStat>& Statistics)
 {
 	TStringBuilder<2048> StatisticsStrBuilder;
-	for (const auto& Stat : Statistics)
+	for (const FGenericShaderStat& Stat : Statistics)
 	{
-		StatisticsStrBuilder << Stat.StatName << ": ";
+		StatisticsStrBuilder << "\n" << Stat.StatName << ": ";
 		Visit([&StatisticsStrBuilder](auto& StoredValue)
  		{
-			StatisticsStrBuilder << StoredValue << "\n";
+			StatisticsStrBuilder << StoredValue;
 		}, Stat.Value);
 	}
 
@@ -486,7 +487,7 @@ void FMaterialStatsUtils::GetRepresentativeInstructionCounts(TArray<FShaderInstr
 						Info.ShaderType = ShaderInfo.ShaderType;
 						Info.ShaderDescription = ShaderInfo.ShaderDescription;
 						Info.InstructionCount = NumInstructions;
-						Info.ShaderStatisticsString = GetShaderString(MaterialShaderMap->GetShaderStatistics(ShaderType));
+						Info.ShaderStatisticsString = GenericShaderStatsToString(MaterialShaderMap->GetShaderStatistics(ShaderType));
 						if (Info.ShaderStatisticsString.Len() == 0)
 						{
 							Info.ShaderStatisticsString = TEXT("n/a");
@@ -529,7 +530,7 @@ void FMaterialStatsUtils::GetRepresentativeInstructionCounts(TArray<FShaderInstr
 									Info.ShaderDescription = ShaderInfo.ShaderDescription;
 									Info.InstructionCount = NumInstructions;
 #if WITH_EDITORONLY_DATA
-									Info.ShaderStatisticsString = GetShaderString(MeshShaderMap->GetShaderStatistics(*MaterialShaderMap, ShaderType));
+									Info.ShaderStatisticsString = GenericShaderStatsToString(MeshShaderMap->GetShaderStatistics(*MaterialShaderMap, ShaderType));
 #endif // WITH_EDITORONLY_DATA
 									if (Info.ShaderStatisticsString.Len() == 0)
 									{
