VulkanRHI: Remove PLATFORM_SUPPORTS_RELAXED_PRECISION from desktop Vulkan SM5 (used only in GL and mobile Vulkan).  Fix landscape water plugin in Vulkan SM5.

#rnx

diff --git a/Engine/Shaders/Public/Platform/Vulkan/VulkanCommon.ush b/Engine/Shaders/Public/Platform/Vulkan/VulkanCommon.ush
--- a/Engine/Shaders/Public/Platform/Vulkan/VulkanCommon.ush
+++ b/Engine/Shaders/Public/Platform/Vulkan/VulkanCommon.ush
@@ -7,7 +7,7 @@
 #pragma once
 
 // Update this GUID to invalidate shader recompilation for only Vulkan shaders
-#pragma message("UESHADERMETADATA_VERSION ABC4AA8D-98C2-48FA-9403-341F12D67621")
+#pragma message("UESHADERMETADATA_VERSION 30B5352B-3DD3-4A44-A103-5D1B9086A805")
 
 #define COMPILER_SUPPORTS_ULONG_TYPES 1
 #define COMPILER_SUPPORTS_UINT64_IMAGE_ATOMICS (1 && COMPILER_SUPPORTS_ULONG_TYPES)
diff --git a/Engine/Source/Developer/VulkanShaderFormat/Private/SpirVShaderCompiler.inl b/Engine/Source/Developer/VulkanShaderFormat/Private/SpirVShaderCompiler.inl
--- a/Engine/Source/Developer/VulkanShaderFormat/Private/SpirVShaderCompiler.inl
+++ b/Engine/Source/Developer/VulkanShaderFormat/Private/SpirVShaderCompiler.inl
@@ -1427,7 +1427,12 @@ static void ModifyCompilerInput(FSpirvShaderCompilerInternalState& InternalState
 	else if (InternalState.IsSM5())
 	{
 		Input.Environment.SetDefine(TEXT("VULKAN_PROFILE_SM5"), 1);
-		Input.Environment.SetDefine(TEXT("PLATFORM_SUPPORTS_RELAXED_PRECISION"), 1);
+		
+		// Only use relaxed precision on mobile platforms
+		if (InternalState.IsAndroid())
+		{
+			Input.Environment.SetDefine(TEXT("PLATFORM_SUPPORTS_RELAXED_PRECISION"), 1);
+		}
 	}
 	Input.Environment.SetDefine(TEXT("row_major"), TEXT(""));
 
@@ -1663,7 +1668,7 @@ static bool CompileShaderGroup(
 {
 	checkf(InternalState.bSupportsBindless && InternalState.bUseBindlessUniformBuffer, TEXT("Ray tracing requires full bindless in Vulkan."));
 
-	// Compile each one of the shader modules seperately and create one big blob for the engine
+	// Compile each one of the shader modules separately and create one big blob for the engine
 	auto CompilePartialExport = [&OriginalPreprocessedShaderSource, &InternalState, &MergedOutput](
 		FSpirvShaderCompilerInternalState::EHitGroupShaderType HitGroupShaderType,
 		const TCHAR* PartialFileExtension,
