Enabling PLATFORM_SUPPORTS_CONSTANTBUFFER_OBJECT on more bindless capable shader formats.

#Synced-CL 49734186

diff --git a/Engine/Source/Developer/Apple/MetalShaderFormat/Private/MetalShaderFormat.cpp b/Engine/Source/Developer/Apple/MetalShaderFormat/Private/MetalShaderFormat.cpp
--- a/Engine/Source/Developer/Apple/MetalShaderFormat/Private/MetalShaderFormat.cpp
+++ b/Engine/Source/Developer/Apple/MetalShaderFormat/Private/MetalShaderFormat.cpp
@@ -435,6 +435,7 @@ public:
 		}
 
 		Input.Environment.SetDefine(TEXT("COMPILER_SUPPORTS_DUAL_SOURCE_BLENDING_SLOT_DECORATION"), (uint32)1);
+		Input.Environment.SetDefine(TEXT("PLATFORM_SUPPORTS_CONSTANTBUFFER_OBJECT"), true);
 	}
 	
 	virtual bool CanCompileBinaryShaders() const override final
diff --git a/Engine/Source/Developer/VulkanShaderFormat/Private/SpirVShaderCompiler.inl b/Engine/Source/Developer/VulkanShaderFormat/Private/SpirVShaderCompiler.inl
--- a/Engine/Source/Developer/VulkanShaderFormat/Private/SpirVShaderCompiler.inl
+++ b/Engine/Source/Developer/VulkanShaderFormat/Private/SpirVShaderCompiler.inl
@@ -1425,6 +1425,7 @@ static void ModifyCompilerInput(FSpirvShaderCompilerInternalState& InternalState
 	{
 		Input.Environment.SetDefine(TEXT("VULKAN_PROFILE_SM6"), 1);
 		Input.Environment.SetDefine(TEXT("PLATFORM_SUPPORTS_CALLABLE_SHADERS"), 1);
+		Input.Environment.SetDefine(TEXT("PLATFORM_SUPPORTS_CONSTANTBUFFER_OBJECT"), true);
 	}
 	else if (InternalState.IsSM5())
 	{
