VulkanRHI: revert ConstantBuffer change (49744009) for Vulkan temporarily

#rnx
#Synced-CL 49763429

diff --git a/Engine/Source/Developer/VulkanShaderFormat/Private/SpirVShaderCompiler.inl b/Engine/Source/Developer/VulkanShaderFormat/Private/SpirVShaderCompiler.inl
--- a/Engine/Source/Developer/VulkanShaderFormat/Private/SpirVShaderCompiler.inl
+++ b/Engine/Source/Developer/VulkanShaderFormat/Private/SpirVShaderCompiler.inl
@@ -1427,7 +1427,7 @@ static void ModifyCompilerInput(FSpirvShaderCompilerInternalState& InternalState
 	{
 		Input.Environment.SetDefine(TEXT("VULKAN_PROFILE_SM6"), 1);
 		Input.Environment.SetDefine(TEXT("PLATFORM_SUPPORTS_CALLABLE_SHADERS"), 1);
-		Input.Environment.SetDefine(TEXT("PLATFORM_SUPPORTS_CONSTANTBUFFER_OBJECT"), true);
+//		Input.Environment.SetDefine(TEXT("PLATFORM_SUPPORTS_CONSTANTBUFFER_OBJECT"), true);
 	}
 	else if (InternalState.IsSM5())
 	{
