VulkanRHI: Add ConstantBuffer patching support to ConvertUBToBindless (will eventually be rolled into ShaderParameterParser)

#rnx

#Synced-CL 49763429

diff --git a/Engine/Source/Developer/VulkanShaderFormat/Private/SpirVShaderCompiler.inl b/Engine/Source/Developer/VulkanShaderFormat/Private/SpirVShaderCompiler.inl
--- a/Engine/Source/Developer/VulkanShaderFormat/Private/SpirVShaderCompiler.inl
+++ b/Engine/Source/Developer/VulkanShaderFormat/Private/SpirVShaderCompiler.inl
@@ -1616,6 +1616,78 @@ static TArray<FString> ConvertUBToBindless(FString& PreprocessedShaderSource,
 		BindlessUBs.SetNum(UBNames.Num());
 	}
 
+	// Patch ConstantBuffers
+	{
+		const FString UniformBufferDeclIdentifier = TEXT("ConstantBuffer");
+
+		int32 SearchIndex = PreprocessedShaderSource.Find(UniformBufferDeclIdentifier, ESearchCase::CaseSensitive, ESearchDir::FromStart, -1);
+		while (SearchIndex != INDEX_NONE)
+		{
+			const TCHAR* TemplateBeginPtr = FCString::Strstr(&PreprocessedShaderSource[SearchIndex + UniformBufferDeclIdentifier.Len()], TEXT("<"));
+			const TCHAR* TemplateEndPtr = FCString::Strstr(TemplateBeginPtr + 1, TEXT(">"));
+
+			FString TemplateName;
+			const TCHAR* TemplateNameEndPtr = ParseHLSLSymbolName(TemplateBeginPtr + 1, TemplateName);
+			check(TemplateName.Len() > 0);
+
+			FString CBName;
+			const TCHAR* CBNameEndPtr = ParseHLSLSymbolName(TemplateEndPtr + 1, CBName);
+			check(CBName.Len() > 0);
+
+
+			int32 CBIndex = INDEX_NONE;
+
+			// Don't patch the VulkanHitGroupSystemParameters system buffer
+			if (CBName != TEXT("VulkanHitGroupSystemParameters"))
+			{
+				if (UBNames.IsEmpty())
+				{
+					CBIndex = BindlessUBs.Add(CBName);
+					check(CBIndex < 16);
+				}
+				else
+				{
+					CBIndex = UBNames.Find(*CBName);
+					if (CBIndex != INDEX_NONE)
+					{
+						BindlessUBs[CBIndex] = CBName;
+					}
+				}
+			}
+
+			const TCHAR* EndDeclPtr = FCString::Strstr(CBNameEndPtr, TEXT(";"));
+			const int32 EndDeclIndex = EndDeclPtr - (*PreprocessedShaderSource);
+
+			int32 NextStartIndex = EndDeclIndex;
+			if (CBIndex != INDEX_NONE)
+			{
+				FString NewDecl;
+
+				// Start with the bindless heap
+				const FString BindlessCBHeapName = FString::Printf(TEXT("%s%d_%s%s"),
+					*SpirvShaderCompiler::kBindlessCBPrefix, CBIndex, *CBName, *SpirvShaderCompiler::kBindlessHeapSuffix);
+				NewDecl += FString::Printf(TEXT("ConstantBuffer<%s> %s[]; "), *TemplateName, *BindlessCBHeapName);
+
+				// Now bring in the CB
+				const TCHAR* BindlessIndexName =
+					(BindlessUniformBufferType == EBindlessUniformBufferType::ShaderRecord) ?
+					TEXT("VulkanHitGroupSystemParameters.BindlessUniformBuffers") :
+					TEXT("VulkanPushConstants.StaticUniformBufferBindings");
+				NewDecl += FString::Printf(TEXT("static const %s %s = %s[%s[%d]];"),
+					*TemplateName, *CBName, *BindlessCBHeapName, BindlessIndexName, CBIndex);
+
+				const int32 OldDeclLen = EndDeclIndex - SearchIndex + 1;
+				PreprocessedShaderSource.RemoveAt(SearchIndex, OldDeclLen, EAllowShrinking::No);
+				PreprocessedShaderSource.InsertAt(SearchIndex, NewDecl);
+
+				NextStartIndex = SearchIndex + NewDecl.Len();
+			}
+
+			SearchIndex = PreprocessedShaderSource.Find(UniformBufferDeclIdentifier, ESearchCase::CaseSensitive, ESearchDir::FromStart, NextStartIndex);
+		}
+	}
+
+	// Patch cbuffers
 	{
 		const FString UniformBufferDeclIdentifier = TEXT("cbuffer");
 
@@ -1663,6 +1735,7 @@ static TArray<FString> ConvertUBToBindless(FString& PreprocessedShaderSource,
 			SearchIndex = PreprocessedShaderSource.Find(UniformBufferDeclIdentifier, ESearchCase::CaseSensitive, ESearchDir::FromStart, NextStartIndex);
 		}
 	}
+
 	return BindlessUBs;
 }
 
diff --git a/Engine/Source/Developer/VulkanShaderFormat/Private/VulkanShaderCompiler.cpp b/Engine/Source/Developer/VulkanShaderFormat/Private/VulkanShaderCompiler.cpp
--- a/Engine/Source/Developer/VulkanShaderFormat/Private/VulkanShaderCompiler.cpp
+++ b/Engine/Source/Developer/VulkanShaderFormat/Private/VulkanShaderCompiler.cpp
@@ -270,11 +270,14 @@ struct FVulkanShaderParameterParserPlatformConfiguration : public FShaderParamet
 	// Fill the global with the value stored in the shader record
 	virtual FString ReplaceGlobal(FStringView FullDeclString, FStringView ParamName) const final
 	{
-		ReplacedGlobals.Add(ParamName, FullDeclString);
-
 		FString NewDecl(FullDeclString);
-		NewDecl = TEXT("static ") + NewDecl;
-		NewDecl.InsertAt(NewDecl.Find(TEXT(";")), FString::Printf(TEXT(" = VulkanHitGroupSystemParameters.Globals.%.*s"), ParamName.Len(), ParamName.GetData()));
+		if (!FullDeclString.Contains(TEXT("ConstantBuffer")))
+		{
+			ReplacedGlobals.Add(ParamName, FullDeclString);
+
+			NewDecl = TEXT("static ") + NewDecl;
+			NewDecl.InsertAt(NewDecl.Find(TEXT(";")), FString::Printf(TEXT(" = VulkanHitGroupSystemParameters.Globals.%.*s"), ParamName.Len(), ParamName.GetData()));
+		}
 		return NewDecl;
 	}
 
@@ -316,8 +319,8 @@ void CompileVulkanShader(const FShaderCompilerInput& Input, const FShaderPreproc
 
 	if (InternalState.bUseBindlessUniformBuffer)
 	{
-		InternalState.ShaderRecordGlobalsSize = ConvertGlobalsToShaderRecord(Input, ShaderParameterParser, ReplacedGlobals, PreprocessedSource, Output);
 		InternalState.AllBindlessUBs = SpirvShaderCompiler::ConvertUBToBindless(PreprocessedSource, SpirvShaderCompiler::EBindlessUniformBufferType::ShaderRecord);
+		InternalState.ShaderRecordGlobalsSize = ConvertGlobalsToShaderRecord(Input, ShaderParameterParser, ReplacedGlobals, PreprocessedSource, Output);
 	}
 
 	if (ShaderParameterParser.DidModifyShader() || InternalState.AllBindlessUBs.Num() || InternalState.ShaderRecordGlobalsSize || InternalState.PushConstantUBs.Num())
