VulkanRHI: add SDCE to Vulkan shader compiler (for improved compilation performance in some substrate configs)

#rnx

#Synced-CL 49847350

diff --git a/Engine/Shaders/Public/Platform/Vulkan/VulkanCommon.ush b/Engine/Shaders/Public/Platform/Vulkan/VulkanCommon.ush
--- a/Engine/Shaders/Public/Platform/Vulkan/VulkanCommon.ush
+++ b/Engine/Shaders/Public/Platform/Vulkan/VulkanCommon.ush
@@ -7,7 +7,7 @@
 #pragma once
 
 // Update this GUID to invalidate shader recompilation for only Vulkan shaders
-#pragma message("UESHADERMETADATA_VERSION EBB4295A-3A82-45E7-89AD-5816B49A55DD")
+#pragma message("UESHADERMETADATA_VERSION B631D136-745C-4AD2-918E-3EBE61B843A9")
 
 #define COMPILER_SUPPORTS_ULONG_TYPES 1
 #define COMPILER_SUPPORTS_UINT64_IMAGE_ATOMICS (1 && COMPILER_SUPPORTS_ULONG_TYPES)
diff --git a/Engine/Source/Developer/VulkanShaderFormat/Private/VulkanShaderCompiler.cpp b/Engine/Source/Developer/VulkanShaderFormat/Private/VulkanShaderCompiler.cpp
--- a/Engine/Source/Developer/VulkanShaderFormat/Private/VulkanShaderCompiler.cpp
+++ b/Engine/Source/Developer/VulkanShaderFormat/Private/VulkanShaderCompiler.cpp
@@ -12,6 +12,7 @@
 #include "VulkanThirdParty.h"
 #include "VulkanBackend.h"
 #include "VulkanShaderResources.h"
+#include "ShaderSDCE.h"
 
 #include "SpirVShaderCompiler.inl"
 
@@ -294,6 +295,9 @@ void CompileVulkanShader(const FShaderCompilerInput& Input, const FShaderPreproc
 	FString EntryPointName = Input.EntryPointName;
 	FString PreprocessedSource(InPreprocessOutput.GetSourceViewWide());
 
+	// Perform symbolic elimination
+	const bool bSDCEMinified = UE::ShaderMinifier::SDCE::MinifyInPlace(Input, InPreprocessOutput, PreprocessedSource);
+
 	TMap<FStringView, FStringView> ReplacedGlobals; // Note: these FStringView point to memory in FShaderParameterParser
 	FVulkanShaderParameterParserPlatformConfiguration PlatformConfiguration(Input, ReplacedGlobals);
 	FShaderParameterParser ShaderParameterParser(PlatformConfiguration);
@@ -323,7 +327,7 @@ void CompileVulkanShader(const FShaderCompilerInput& Input, const FShaderPreproc
 		InternalState.ShaderRecordGlobalsSize = ConvertGlobalsToShaderRecord(Input, ShaderParameterParser, ReplacedGlobals, PreprocessedSource, Output);
 	}
 
-	if (ShaderParameterParser.DidModifyShader() || InternalState.AllBindlessUBs.Num() || InternalState.ShaderRecordGlobalsSize || InternalState.PushConstantUBs.Num())
+	if (ShaderParameterParser.DidModifyShader() || bSDCEMinified || InternalState.AllBindlessUBs.Num() || InternalState.ShaderRecordGlobalsSize || InternalState.PushConstantUBs.Num())
 	{
 		Output.ModifiedShaderSource = PreprocessedSource;
 	}
