Vulkan Bindless Uniform Buffers
- "manual" bindless uniform buffers, ie ones using UB_BINDLESS_DECL, need to be handled the same way as other bindless resources.
- Bumping the Vulkan SM6 version.

#Synced-CL 52522322

diff --git a/Engine/Source/Developer/VulkanShaderFormat/Private/SpirVShaderCompiler.inl b/Engine/Source/Developer/VulkanShaderFormat/Private/SpirVShaderCompiler.inl
--- a/Engine/Source/Developer/VulkanShaderFormat/Private/SpirVShaderCompiler.inl
+++ b/Engine/Source/Developer/VulkanShaderFormat/Private/SpirVShaderCompiler.inl
@@ -597,6 +597,7 @@ static void GatherSpirvReflectionBindings(
 		MoveBindlessHeaps(OutBindings.TBufferSRVs, FShaderParameterParser::kBindlessSRVArrayPrefix, VulkanBindless::BindlessUniformTexelBufferSet, OutBindings.NumBindlessResources);
 		MoveBindlessHeaps(OutBindings.TBufferUAVs, FShaderParameterParser::kBindlessUAVArrayPrefix, VulkanBindless::BindlessStorageTexelBufferSet, OutBindings.NumBindlessResources);
 		MoveBindlessHeaps(OutBindings.AccelerationStructures, FShaderParameterParser::kBindlessSRVArrayPrefix, VulkanBindless::BindlessAccelerationStructureSet, OutBindings.NumBindlessResources);
+		MoveBindlessHeaps(OutBindings.UniformBuffers, FShaderParameterParser::kBindlessCBVArrayPrefix, VulkanBindless::BindlessUniformBufferSet, OutBindings.NumBindlessResources);
 
 		// Move uniform buffers to the correct set
 		{
diff --git a/Engine/Source/Developer/VulkanShaderFormat/Private/VulkanShaderFormat.cpp b/Engine/Source/Developer/VulkanShaderFormat/Private/VulkanShaderFormat.cpp
--- a/Engine/Source/Developer/VulkanShaderFormat/Private/VulkanShaderFormat.cpp
+++ b/Engine/Source/Developer/VulkanShaderFormat/Private/VulkanShaderFormat.cpp
@@ -24,7 +24,7 @@ extern void CompileVulkanShader(
 
 static const FGuid UE_SHADER_VULKAN_ES3_1_VER = FGuid("B84F72C8-3ECD-411E-993C-D7C7CEE26F28");
 static const FGuid UE_SHADER_VULKAN_SM5_VER = FGuid("0715D8EE-9907-4A25-93AD-A3902C8E069A");
-static const FGuid UE_SHADER_VULKAN_SM6_VER = FGuid("C5161730-83C6-40AF-A990-78CD4C1581DB");
+static const FGuid UE_SHADER_VULKAN_SM6_VER = FGuid("506C4300-261F-4B9B-97AE-3921B3CFFF34");
 
 class FShaderFormatVulkan : public UE::ShaderCompilerCommon::FBaseShaderFormat
 {
